Dataflow:
- Moved DataflowVolumetric plugin to /Engine/Plugins/Experimental

diff --git a/Engine/Plugins/Experimental/DataflowVolumetricPlugin/DataflowVolumetricPlugin.uplugin b/Engine/Plugins/Experimental/DataflowVolumetricPlugin/DataflowVolumetricPlugin.uplugin
new file mode 100644
--- /dev/null
+++ b/Engine/Plugins/Experimental/DataflowVolumetricPlugin/DataflowVolumetricPlugin.uplugin
@@ -0,0 +1,41 @@
+{
+    "FileVersion" : 3,
+    "Version" : 1,
+    "VersionName" : "0.1",
+    "FriendlyName" : "DataflowVolumetric",
+    "Description" : "Adds volumetric support to Dataflow",
+    "Category": "Dataflow",
+    "CreatedBy" : "Epic Games, Inc.",
+    "CreatedByURL" : "https://epicgames.com",
+    "DocsURL" : "",
+    "MarketplaceURL" : "",
+    "SupportURL" : "",
+    "CanContainContent" : true,
+    "IsBetaVersion" : false,
+    "IsExperimentalVersion" : true,
+    "Installed" : false,
+	"Modules": [
+		{
+			"Name": "DataflowVolumeCore",
+			"Type": "Editor",
+			"LoadingPhase": "Default",
+			"PlatformAllowList": [
+				"Win64", "Linux", "Mac"
+			]
+		},
+		{
+			"Name": "DataflowVolumeNodes",
+			"Type": "Editor",
+			"LoadingPhase": "Default",
+			"PlatformAllowList": [
+				"Win64", "Linux", "Mac"
+			]
+		}
+	],
+	"Plugins": [
+		{
+			"Name": "Dataflow",
+			"Enabled": true
+		}
+	]
+}
diff --git a/Engine/Plugins/Experimental/DataflowVolumetricPlugin/Source/DataflowVolumeCore/DataflowVolumeCore.Build.cs b/Engine/Plugins/Experimental/DataflowVolumetricPlugin/Source/DataflowVolumeCore/DataflowVolumeCore.Build.cs
new file mode 100644
--- /dev/null
+++ b/Engine/Plugins/Experimental/DataflowVolumetricPlugin/Source/DataflowVolumeCore/DataflowVolumeCore.Build.cs
@@ -0,0 +1,49 @@
+// Copyright Epic Games, Inc. All Rights Reserved.
+
+namespace UnrealBuildTool.Rules
+{
+	public class DataflowVolumeCore : ModuleRules
+	{
+        public DataflowVolumeCore(ReadOnlyTargetRules Target) : base(Target)
+		{
+			// For boost:: and TBB:: code
+			bUseRTTI = true;
+
+			PublicDependencyModuleNames.AddRange(
+				new string[]
+				{
+				}
+			);
+
+			PrivateDependencyModuleNames.AddRange(
+				new string[] 
+				{
+					"DynamicMesh",
+					"MeshConversion",
+					"Core",
+					"CoreUObject",
+					"ChaosCore",
+					"Chaos",
+					"DataflowCore",
+					"DataflowEngine",
+					"DataflowEnginePlugin",
+					"DataflowNodes",
+					"GeometryCollectionEngine",
+					"GeometryCore",
+					"Voronoi",
+					"GeometryFramework",
+					"MeshDescription",
+					"StaticMeshDescription",
+					"PlanarCut",
+					"Engine",
+					"DataflowEditor"
+				}
+			);
+
+			AddEngineThirdPartyPrivateStaticDependencies(Target,
+				"IntelTBB",
+				"OpenVDB"
+			);
+		}
+	}
+}
diff --git a/Engine/Plugins/Experimental/DataflowVolumetricPlugin/Source/DataflowVolumeCore/Internal/Dataflow/DataflowBoolVolume.h b/Engine/Plugins/Experimental/DataflowVolumetricPlugin/Source/DataflowVolumeCore/Internal/Dataflow/DataflowBoolVolume.h
new file mode 100644
--- /dev/null
+++ b/Engine/Plugins/Experimental/DataflowVolumetricPlugin/Source/DataflowVolumeCore/Internal/Dataflow/DataflowBoolVolume.h
@@ -0,0 +1,39 @@
+// Copyright Epic Games, Inc. All Rights Reserved.
+
+#pragma once 
+
+#include "DataflowVolume.h"
+
+#include "DataflowBoolVolume.generated.h"
+
+namespace UE::DataflowVolume::Private
+{
+	class FDataflowBoolVolumeImpl;
+}
+
+/* --------------------------------------------------------------------------------------------------------------- */
+/* FDataflowBoolVolume */
+/* --------------------------------------------------------------------------------------------------------------- */
+
+USTRUCT()
+struct FDataflowBoolVolume : public FDataflowVolume
+{
+	inline static const FName TypeName = "BoolVolume";
+
+	GENERATED_USTRUCT_BODY()
+
+public:
+	FDataflowBoolVolume() {}
+	virtual ~FDataflowBoolVolume() {}
+
+	TSharedPtr<const UE::DataflowVolume::Private::FDataflowBoolVolumeImpl> GetBoolVolumeImpl() const;
+	TSharedPtr<UE::DataflowVolume::Private::FDataflowBoolVolumeImpl> GetBoolVolumeImpl();
+
+	FVector GetVoxelSize() const;
+	FBox GetVolumeBoundingBox() const;
+
+private:
+	FDataflowBoolVolume(TSharedPtr<UE::DataflowVolume::Private::FDataflowBoolVolumeImpl> InBoolVolume);
+};
+
+
diff --git a/Engine/Plugins/Experimental/DataflowVolumetricPlugin/Source/DataflowVolumeCore/Internal/Dataflow/DataflowFloatVectorVolume.h b/Engine/Plugins/Experimental/DataflowVolumetricPlugin/Source/DataflowVolumeCore/Internal/Dataflow/DataflowFloatVectorVolume.h
new file mode 100644
--- /dev/null
+++ b/Engine/Plugins/Experimental/DataflowVolumetricPlugin/Source/DataflowVolumeCore/Internal/Dataflow/DataflowFloatVectorVolume.h
@@ -0,0 +1,52 @@
+// Copyright Epic Games, Inc. All Rights Reserved.
+
+#pragma once 
+
+#include "Dataflow/DataflowVolume.h"
+#include "Dataflow/DataflowVolumeNodeEnums.h"
+
+#include "DataflowFloatVectorVolume.generated.h"
+
+namespace UE::DataflowVolume::Private
+{
+	class FDataflowFloatVectorVolumeImpl;
+}
+
+struct FDataflowFloatVolume;
+struct FDataflowBoolVolume;
+
+/* --------------------------------------------------------------------------------------------------------------- */
+/* FDataflowFloatVectorVolume */
+/* --------------------------------------------------------------------------------------------------------------- */
+
+USTRUCT()
+struct FDataflowFloatVectorVolume : public FDataflowVolume
+{
+	inline static const FName TypeName = "FloatVectorVolume";
+
+	GENERATED_USTRUCT_BODY()
+
+public:
+	friend struct FDataflowFloatVolume;
+
+	FDataflowFloatVectorVolume() {}
+	virtual ~FDataflowFloatVectorVolume() {}
+
+	TSharedPtr<const UE::DataflowVolume::Private::FDataflowFloatVectorVolumeImpl> GetFloatVectorVolumeImpl() const;
+	TSharedPtr<UE::DataflowVolume::Private::FDataflowFloatVectorVolumeImpl> GetFloatVectorVolumeImpl();
+
+	DATAFLOWVOLUMECORE_API FVector GetVoxelSize() const;
+	DATAFLOWVOLUMECORE_API FBox GetVolumeBoundingBox() const;
+	DATAFLOWVOLUMECORE_API FString GetVectorType() const;
+
+	DATAFLOWVOLUMECORE_API void VolumeSample(const TArray<FVector>& InPoints, TArray<FVector>& OutValues) const;
+
+	DATAFLOWVOLUMECORE_API FDataflowFloatVolume ComputeDivergence(const EDataflowVolumeAnalysisOutputName OutputName, const FString& CustomName, const FDataflowBoolVolume* MaskVolume = nullptr) const;
+	DATAFLOWVOLUMECORE_API FDataflowFloatVolume ComputeMagnitude(const EDataflowVolumeAnalysisOutputName OutputName, const FString& CustomName, const FDataflowBoolVolume* MaskVolume = nullptr) const;
+	DATAFLOWVOLUMECORE_API FDataflowFloatVectorVolume ComputeCurl(const EDataflowVolumeAnalysisOutputName OutputName, const FString& CustomName, const FDataflowBoolVolume* MaskVolume = nullptr) const;
+	DATAFLOWVOLUMECORE_API FDataflowFloatVectorVolume ComputeNormalize(const EDataflowVolumeAnalysisOutputName OutputName, const FString& CustomName, const FDataflowBoolVolume* MaskVolume = nullptr) const;
+
+private:
+	FDataflowFloatVectorVolume(TSharedPtr<UE::DataflowVolume::Private::FDataflowFloatVectorVolumeImpl> InFloatVectorVolume);
+};
+
diff --git a/Engine/Plugins/Experimental/DataflowVolumetricPlugin/Source/DataflowVolumeCore/Internal/Dataflow/DataflowFloatVolume.h b/Engine/Plugins/Experimental/DataflowVolumetricPlugin/Source/DataflowVolumeCore/Internal/Dataflow/DataflowFloatVolume.h
new file mode 100644
--- /dev/null
+++ b/Engine/Plugins/Experimental/DataflowVolumetricPlugin/Source/DataflowVolumeCore/Internal/Dataflow/DataflowFloatVolume.h
@@ -0,0 +1,152 @@
+// Copyright Epic Games, Inc. All Rights Reserved.
+
+#pragma once 
+
+#include "Dataflow/DataflowVolume.h"
+#include "DataflowVolumeNodeEnums.h"
+
+#include "DataflowFloatVolume.generated.h"
+
+namespace UE::DataflowVolume::Private
+{
+	class FDataflowFloatVolumeImpl;
+}
+
+struct FDataflowIntVolume;
+struct FDataflowBoolVolume;
+struct FDataflowFloatVectorVolume;
+struct FMeshDescription;
+
+/* --------------------------------------------------------------------------------------------------------------- */
+/* FDataflowFloatVolume */
+/* --------------------------------------------------------------------------------------------------------------- */
+USTRUCT()
+struct FDataflowFloatVolume : public FDataflowVolume
+{
+	inline static const FName TypeName = "FloatVolume";
+
+	GENERATED_USTRUCT_BODY()
+
+public:
+	FDataflowFloatVolume() {}
+	virtual ~FDataflowFloatVolume() {}
+
+	TSharedPtr<const UE::DataflowVolume::Private::FDataflowFloatVolumeImpl> GetFloatVolumeImpl() const;
+	TSharedPtr<UE::DataflowVolume::Private::FDataflowFloatVolumeImpl> GetFloatVolumeImpl();
+
+	DATAFLOWVOLUMECORE_API FVector GetVoxelSize() const;
+	DATAFLOWVOLUMECORE_API FBox GetVolumeBoundingBox() const;
+
+	// Create a sphere SDF
+	DATAFLOWVOLUMECORE_API static FDataflowFloatVolume CreateSphereSDF(const float InVoxelSize, const float InRadius, const FVector& InCenter);
+
+	/** Create a cube SDF */
+	DATAFLOWVOLUMECORE_API static FDataflowFloatVolume CreateCubeSDF(const float InScale, const FVector& InCenter, const float InVoxelSize);
+
+	/** Create different types of solid SDF */
+	DATAFLOWVOLUMECORE_API static FDataflowFloatVolume CreatePlatonicSolidSDF(const int32 InFaceCount, const float InScale, const FVector& InCenter, const float InVoxelSize);
+
+	/** Create a volume from MeshDescription */
+	DATAFLOWVOLUMECORE_API static FDataflowFloatVolume CreateVolumeFromMeshDescription(
+		const FMeshDescription& InMeshDescription,
+		const FTransform& InTransform,
+		const float InVoxelSize,
+		const FDataflowVolume& InReferenceGrid,
+		const EDataflowVolumeOutputType InOutputType,
+		const FString& InGridName,
+		const bool InUseWorldSpaceUnits,
+		const float InExteriorBand,
+		const float InInteriorBand,
+		const int32 InExteriorBandVoxels,
+		const int32 InInteriorBandVoxels,
+		const bool bFillInterior,
+		const bool bPreserveHoles,
+		const float InIsoValue,
+		FDataflowIntVolume& OutFaceIndexVolume);
+
+	/** Compute sphere packing in an SDF volume */
+	DATAFLOWVOLUMECORE_API void VolumeToSpheres(
+		const int32 InMinSphereCount,
+		const int32 InMaxSphereCount,
+		const bool InOverlapping,
+		const float InMinRadius,
+		const float InMaxRadius,
+		const float InIsoValue,
+		const int32 InInstanceCount,
+		TArray<FVector>& SphereCenters,
+		TArray<float>& SphereRadii) const;
+
+	/** Convert SDF to Fog Volume */
+	DATAFLOWVOLUMECORE_API FDataflowFloatVolume CovertSDFToFogVolume(const bool InPruneTolerance, const float InTolerance, const bool InFloodFillOutput, const bool InActivateInteriror) const;
+
+	/** Convert SDF to Fog Volume */
+	DATAFLOWVOLUMECORE_API FDataflowFloatVolume CovertFogVolumeToSDF(
+		const float InFogIsoValue,
+		const bool InPruneTolerance,
+		const float InTolerance,
+		const bool InFloodFillOutput,
+		const bool InActivateInteriror) const;
+
+	/** Convert volume to MeshDescription */
+	DATAFLOWVOLUMECORE_API void ConvertVolumeToMeshDescription(const float InIsoValue, const float InAdaptivity, FMeshDescription& OutMeshDescription) const;
+
+	/** Compute uniform point scatter in an SDF volume */
+	DATAFLOWVOLUMECORE_API void UniformVolumeScatter(
+		const int32 InMinNumberOfPoints,
+		const int32 InMaxNumberOfPoints,
+		const int32 InRandomSeed,
+		const float InIsoValue,
+		const float InSpread,
+		TArray<FVector>& OutPoints) const;
+
+	/** Compute dense uniform point scatter in an SDF volume */
+	DATAFLOWVOLUMECORE_API void DenseUniformVolumeScatter(
+		const float InMinNumberOfPointsPerVoxel,
+		const float InMaxNumberOfPointsPerVoxel,
+		const int32 InRandomSeed,
+		const float InIsoValue,
+		const float InSpread,
+		TArray<FVector>& OutPoints) const;
+
+	/** Compute a non uniform point scatter in an SDF volume */
+	DATAFLOWVOLUMECORE_API void NonUniformVolumeScatter(
+		const float InMinNumberOfPointsPerVoxel,
+		const float InMaxNumberOfPointsPerVoxel,
+		const int32 InRandomSeed,
+		const float InIsoValue,
+		const float InSpread,
+		TArray<FVector>& OutPoints) const;
+
+	/** Sample a volume */
+	DATAFLOWVOLUMECORE_API void VolumeSample(const TArray<FVector>& InPoints, TArray<float>& OutValues) const;
+
+	/** Compute gradient */
+	DATAFLOWVOLUMECORE_API FDataflowFloatVectorVolume ComputeGradient(const EDataflowVolumeAnalysisOutputName OutputName, const FString& CustomName, const FDataflowBoolVolume* MaskVolume = nullptr) const;
+
+	/** Compute curvature */
+	DATAFLOWVOLUMECORE_API FDataflowFloatVolume ComputeCurvature(const EDataflowVolumeAnalysisOutputName OutputName, const FString& CustomName, const FDataflowBoolVolume* MaskVolume = nullptr) const;
+
+	/** Compute Laplacian */
+	DATAFLOWVOLUMECORE_API FDataflowFloatVolume ComputeLaplacian(const EDataflowVolumeAnalysisOutputName OutputName, const FString& CustomName, const FDataflowBoolVolume* MaskVolume = nullptr) const;
+
+	/** Compute ClosestPoint */
+	DATAFLOWVOLUMECORE_API FDataflowFloatVectorVolume ComputeClosestPoint(const EDataflowVolumeAnalysisOutputName OutputName, const FString& CustomName, const FDataflowBoolVolume* MaskVolume = nullptr) const;
+
+	/** Combine SDF volumes */
+	DATAFLOWVOLUMECORE_API FDataflowFloatVolume VolumeCombine(
+		const FDataflowFloatVolume& InVolumeB,
+		const EDataflowVolumeSDFCombineOperation InOperation,
+		const float InMultiplierA,
+		const float InMultiplierB,
+		const EDataflowVolumeSDFCombineResample InResample,
+		const EDataflowVolumeSDFCombineInterpolation InInterpolation,
+		const bool InPruneDegenerateTiles,
+		const float InPruneTol) const;
+
+private:
+	friend struct FDataflowFloatVectorVolume;
+	friend class UE::DataflowVolume::Private::FDataflowFloatVolumeImpl;
+
+	FDataflowFloatVolume(TSharedPtr<UE::DataflowVolume::Private::FDataflowFloatVolumeImpl> InFloatVolume);
+};
+
diff --git a/Engine/Plugins/Experimental/DataflowVolumetricPlugin/Source/DataflowVolumeCore/Internal/Dataflow/DataflowIntVolume.h b/Engine/Plugins/Experimental/DataflowVolumetricPlugin/Source/DataflowVolumeCore/Internal/Dataflow/DataflowIntVolume.h
new file mode 100644
--- /dev/null
+++ b/Engine/Plugins/Experimental/DataflowVolumetricPlugin/Source/DataflowVolumeCore/Internal/Dataflow/DataflowIntVolume.h
@@ -0,0 +1,45 @@
+// Copyright Epic Games, Inc. All Rights Reserved.
+
+#pragma once 
+
+#include "DataflowVolume.h"
+
+#include "DataflowIntVolume.generated.h"
+
+namespace UE::DataflowVolume::Private
+{
+	class FDataflowIntVolumeImpl;
+	class FDataflowFloatVolumeImpl;
+}
+
+/* --------------------------------------------------------------------------------------------------------------- */
+/* FDataflowIntVolume */
+/* --------------------------------------------------------------------------------------------------------------- */
+USTRUCT()
+struct FDataflowIntVolume : public FDataflowVolume
+{
+	inline static const FName TypeName = "IntVolume";
+
+	GENERATED_USTRUCT_BODY()
+
+public:
+	FDataflowIntVolume() {}
+	virtual ~FDataflowIntVolume() {}
+
+	TSharedPtr<const UE::DataflowVolume::Private::FDataflowIntVolumeImpl> GetIntVolume() const;
+	TSharedPtr<UE::DataflowVolume::Private::FDataflowIntVolumeImpl> GetIntVolume();
+
+	DATAFLOWVOLUMECORE_API FVector GetVoxelSize() const;
+	DATAFLOWVOLUMECORE_API FBox GetVolumeBoundingBox() const;
+
+	/** Sample a volume */
+	DATAFLOWVOLUMECORE_API void VolumeSample(const TArray<FVector>& InPoints, TArray<int32>& OutValues) const;
+
+private:
+	friend struct FDataflowFloatVolume;
+
+	FDataflowIntVolume(TSharedPtr<UE::DataflowVolume::Private::FDataflowIntVolumeImpl> InIntVolume);
+
+	void SetIntVolume(TSharedPtr<UE::DataflowVolume::Private::FDataflowIntVolumeImpl> InIntVolume);
+};
+
diff --git a/Engine/Plugins/Experimental/DataflowVolumetricPlugin/Source/DataflowVolumeCore/Internal/Dataflow/DataflowVolume.h b/Engine/Plugins/Experimental/DataflowVolumetricPlugin/Source/DataflowVolumeCore/Internal/Dataflow/DataflowVolume.h
new file mode 100644
--- /dev/null
+++ b/Engine/Plugins/Experimental/DataflowVolumetricPlugin/Source/DataflowVolumeCore/Internal/Dataflow/DataflowVolume.h
@@ -0,0 +1,93 @@
+// Copyright Epic Games, Inc. All Rights Reserved.
+
+#pragma once 
+
+#include "DataflowVolume.generated.h"
+
+namespace UE::DataflowVolume::Private
+{
+	class FDataflowVolumeImpl;
+}
+
+/* --------------------------------------------------------------------------------------------------------------- */
+/* FDataflowVolume */
+/* --------------------------------------------------------------------------------------------------------------- */
+USTRUCT()
+struct FDataflowVolume
+{
+	GENERATED_USTRUCT_BODY()
+
+public:
+	FDataflowVolume() {}
+	virtual ~FDataflowVolume() {}
+
+	DATAFLOWVOLUMECORE_API void GetActiveVoxels(TArray<FBox>& OutActiveVoxels) const;
+
+	/** Prints Volume info into string */
+	DATAFLOWVOLUMECORE_API FString VolumeInfo() const;
+	FName GetType() const { return Type; }
+
+	/*
+	
+	 How to use this:
+	 
+	 GetValue() for AnyType returns the storage type
+	 
+	 @code
+	 const FDataflowVolume& Volume = GetValue(Context, &MyVolume);
+	
+	 if (const FDataflowFloatVolume* FloatVolume = Volume.Cast<FDataflowFloatVolume>())
+	 {
+	    ...
+	 }	
+	 else if (const FDataflowIntVolume* IntVolume = Volume.Cast<FDataflowIntVolume>())
+	 {
+	    ...
+	 }
+	 else if (const FDataflowFloatVectorVolume* FloatVectorVolume = Volume.Cast<FDataflowFloatVectorVolume>())
+	 {
+	    ...
+	 }
+	 @endcode
+
+	*/
+
+	template<typename T>
+	T* Cast()
+	{
+		if constexpr (std::is_base_of_v<FDataflowVolume, T>)
+		{
+			if (Type == T::TypeName)
+			{
+				return static_cast<T*>(this);
+			}
+		}
+
+		return nullptr;
+	}
+
+	template<typename T>
+	const T* Cast() const
+	{
+		if constexpr (std::is_base_of_v<FDataflowVolume, T>)
+		{
+			if (Type == T::TypeName)
+			{
+				return static_cast<const T*>(this);
+			}
+		}
+
+		return nullptr;
+	}
+
+	TSharedPtr<const UE::DataflowVolume::Private::FDataflowVolumeImpl> GetVolume() const { return Volume; }
+
+protected:
+	FDataflowVolume(TSharedPtr<UE::DataflowVolume::Private::FDataflowVolumeImpl> InVolume, const FName InType);
+
+	TSharedPtr<UE::DataflowVolume::Private::FDataflowVolumeImpl> Volume;
+
+private:
+	FName Type;
+};
+
diff --git a/Engine/Plugins/Experimental/DataflowVolumetricPlugin/Source/DataflowVolumeCore/Internal/Dataflow/DataflowVolumeAnyType.h b/Engine/Plugins/Experimental/DataflowVolumetricPlugin/Source/DataflowVolumeCore/Internal/Dataflow/DataflowVolumeAnyType.h
new file mode 100644
--- /dev/null
+++ b/Engine/Plugins/Experimental/DataflowVolumetricPlugin/Source/DataflowVolumeCore/Internal/Dataflow/DataflowVolumeAnyType.h
@@ -0,0 +1,59 @@
+// Copyright Epic Games, Inc. All Rights Reserved.
+
+#pragma once
+
+#include "Dataflow/DataflowAnyType.h"
+#include "Dataflow/DataflowTypePolicy.h"
+#include "Dataflow/DataflowVolume.h"
+#include "Dataflow/DataflowFloatVolume.h"
+#include "Dataflow/DataflowIntVolume.h"
+#include "Dataflow/DataflowFloatVectorVolume.h"
+
+#include "DataflowVolumeAnyType.generated.h"
+
+UE_DATAFLOW_POLICY_DECLARE_TYPENAME(FDataflowVolume)
+UE_DATAFLOW_POLICY_DECLARE_TYPENAME(FDataflowFloatVolume)
+UE_DATAFLOW_POLICY_DECLARE_TYPENAME(FDataflowIntVolume)
+UE_DATAFLOW_POLICY_DECLARE_TYPENAME(FDataflowFloatVectorVolume)
+
+struct FDataflowVolumeTypePolicy :
+	public TDataflowMultiTypePolicy<FDataflowFloatVolume, FDataflowIntVolume, FDataflowFloatVectorVolume>
+{
+};
+
+USTRUCT()
+struct FDataflowVolumeTypes : public FDataflowAnyType
+{
+	using FPolicyType = FDataflowVolumeTypePolicy;
+	using FStorageType = FDataflowVolume;
+
+	GENERATED_USTRUCT_BODY()
+
+	UPROPERTY()
+	FDataflowVolume Value;
+};
+
+template <>
+struct FDataflowConverter<FDataflowVolume>
+{
+	template <typename TFromType>
+	static void From(const TFromType& From, FDataflowVolume& To)
+	{
+		To = From; // from typed to general FDataflowVolume
+	}
+
+	template <typename TToType>
+	static void To(const FDataflowVolume& From, TToType& To)
+	{
+		if (const TToType& Typed = From.Cast<TToType>())
+		{
+			To = Typed;
+		}
+	}
+};
+
+namespace UE::Dataflow
+{
+	DATAFLOWVOLUMECORE_API void RegisterVolumeAnyTypes();
+}
+
diff --git a/Engine/Plugins/Experimental/DataflowVolumetricPlugin/Source/DataflowVolumeCore/Internal/Dataflow/DataflowVolumeNodeEnums.h b/Engine/Plugins/Experimental/DataflowVolumetricPlugin/Source/DataflowVolumeCore/Internal/Dataflow/DataflowVolumeNodeEnums.h
new file mode 100644
--- /dev/null
+++ b/Engine/Plugins/Experimental/DataflowVolumetricPlugin/Source/DataflowVolumeCore/Internal/Dataflow/DataflowVolumeNodeEnums.h
@@ -0,0 +1,170 @@
+// Copyright Epic Games, Inc. All Rights Reserved.
+
+#pragma once 
+
+#include "DataflowVolumeNodeEnums.generated.h"
+
+UENUM(BlueprintType)
+enum class EDataflowVolumeOutputType : uint8
+{
+	SDF UMETA(DisplayName = "Signed Distance Field"),
+	USDF UMETA(DisplayName = "Unsigned Distance Field"),
+	FogVolume UMETA(DisplayName = "Fog Volume"),
+};
+
+UENUM(BlueprintType)
+enum class EDataflowVolumeSDFCombineOperation : uint8
+{
+	CopyA UMETA(DisplayName = "Copy A"),
+	CopyB UMETA(DisplayName = "Copy B"),
+	InvertA UMETA(DisplayName = "Invert A"),
+	Add UMETA(DisplayName = "Add"),
+	Subtract UMETA(DisplayName = "Subtract"),
+	Multiply UMETA(DisplayName = "Multiply"),
+	Divide UMETA(DisplayName = "Divide"),
+	Minimum UMETA(DisplayName = "Minimum"),
+	Maximum UMETA(DisplayName = "Maximum"),
+	CompATimesB UMETA(DisplayName = "(1 - A) * B"),
+	APlusCompATimesB UMETA(DisplayName = "A + (1 - A) * B"),
+	SDFUnion UMETA(DisplayName = "SDF Union"),
+	SDFIntersect UMETA(DisplayName = "SDF Intersection"),
+	SDFDifference UMETA(DisplayName = "SDF Difference"),
+	ReplaceWithActive UMETA(DisplayName = "Replace A with Active B"),
+	ActivityUnion UMETA(DisplayName = "Activity Union"),
+	ActivityIntersect UMETA(DisplayName = "Activity Intersection"),
+	ActivityDifference UMETA(DisplayName = "Activity Difference")
+};
+
+UENUM(BlueprintType)
+enum class EDataflowVolumeSDFCombineResample : uint8
+{
+	Off UMETA(DisplayName = "Off"),
+	BMatchA UMETA(DisplayName = "B to Match A"),
+	AMatchB UMETA(DisplayName = "A to Match B"),
+	HiresMatchLores UMETA(DisplayName = "Higher-res to Match Lower-res"),
+	LoresMatchHires UMETA(DisplayName = "Lower-res to Match Higher-res")
+};
+
+UENUM(BlueprintType)
+enum class EDataflowVolumeSDFCombineInterpolation : uint8
+{
+	Nearest UMETA(DisplayName = "Nearest"),
+	Linear UMETA(DisplayName = "Linear"),
+	Quadratic UMETA(DisplayName = "Quadratic")
+};
+
+UENUM(BlueprintType)
+enum class EDataflowVolumeConvertSDFTo : uint8
+{
+	Volume UMETA(DisplayName = "Volume"),
+	Collection UMETA(DisplayName = "Collection"),
+};
+
+UENUM(BlueprintType)
+enum class EDataflowVolumeConvertSDFGridClass : uint8
+{
+	NoChange UMETA(DisplayName = "No Change"),
+	SDFToFog UMETA(DisplayName = "SDF -> Fog Volume"),
+	FogToSDF UMETA(DisplayName = "Fog Volume -> SDF"),
+};
+
+UENUM(BlueprintType)
+enum class EDataflowVolumeConvertSDFGridType : uint8
+{
+	NoChange UMETA(DisplayName = "No Change"),
+	Float UMETA(DisplayName = "Float"),
+	Int UMETA(DisplayName = "Int"),
+	Bool UMETA(DisplayName = "Bool"),
+	VectorFloat UMETA(DisplayName = "Vector Float"),
+	VectorInt UMETA(DisplayName = "Vector Int"),
+};
+
+UENUM(BlueprintType)
+enum class EDataflowVolumeActiveVoxelsDisplayType : uint8
+{
+	Points UMETA(DisplayName = "Points"),
+	WwireframeBoxes UMETA(DisplayName = "Wireframe Boxes"),
+	SolidBoxes UMETA(DisplayName = "Solid Boxes"),
+};
+
+UENUM(BlueprintType)
+enum class EDataflowVolumePlatonicSolidType : uint8
+{
+	Tetrahedron UMETA(DisplayName = "Tetrahedron"),
+	Cube UMETA(DisplayName = "Cube"),
+	Octahedron UMETA(DisplayName = "Octahedron"),
+	Dodecahedron UMETA(DisplayName = "Dodecahedron"),
+	Icosahedron UMETA(DisplayName = "Icosahedron"),
+};
+
+UENUM(BlueprintType)
+enum class EDataflowVolumeSampleAttribute : uint8
+{
+	Phi UMETA(DisplayName = "Vertex/Phi"),
+	Gradient UMETA(DisplayName = "Vertex/Gradient"),
+	VertexColor UMETA(DisplayName = "Vertex/Color"),
+	VertexClosestFaceID UMETA(DisplayName = "Vertex/ClosestFaceID"),
+	VertexUV UMETA(DisplayName = "Vertex/UV"),
+	FaceMaterialID UMETA(DisplayName = "Face/MaterialID"),
+	VertexCustomAttribute UMETA(DisplayName = "Vertex/Custom Attribute"),
+	FaceCustomAttribute UMETA(DisplayName = "Face/Custom Attribute"),
+};
+
+UENUM(BlueprintType)
+enum class EDataflowVolumeSampleType : uint8
+{
+	Float UMETA(DisplayName = "Float"),
+	Int UMETA(DisplayName = "Int"),
+	FloatVector UMETA(DisplayName = "FloatVector"),
+};
+
+UENUM(BlueprintType)
+enum class EDataflowVolumeSampleLocation : uint8
+{
+	Vertices UMETA(DisplayName = "Vertices"),
+	FaceCenters UMETA(DisplayName = "Face Centers"),
+};
+
+UENUM(BlueprintType)
+enum class EDataflowVolumeSliceMethod : uint8
+{
+	Points UMETA(DisplayName = "Points"),
+};
+
+UENUM(BlueprintType)
+enum class EDataflowVolumeSlicePlane : uint8
+{
+	XYPlane UMETA(DisplayName = "XY Plane"),
+	YZPlane UMETA(DisplayName = "YZ Plane"),
+	ZXPlane UMETA(DisplayName = "ZX Plane"),
+};
+
+UENUM(BlueprintType)
+enum class EDataflowVolumeSliceRamp : uint8
+{
+	InfraRed UMETA(DisplayName = "Infra-Red"),
+	WhiteToRed UMETA(DisplayName = "White to Red"),
+	Grayscale UMETA(DisplayName = "Grayscale"),
+	BlackBody UMETA(DisplayName = "Blackbody"),
+	Custom UMETA(DisplayName = "Custom"),
+};
+
+UENUM(BlueprintType)
+enum class EDataflowVolumeAnalysisOperator : uint8
+{
+	Gradient UMETA(DisplayName = "Gradient (Scalar->Vector)"),
+	Curvature UMETA(DisplayName = "Curvature (Scalar->Scalar)"),
+	Laplacian UMETA(DisplayName = "Laplacian (Scalar->Scalar)"),
+	ClosestPoint UMETA(DisplayName = "Closest Point (Scalar->Vector)"),
+	Divergence UMETA(DisplayName = "Divergence (Vector->Scalar)"),
+	Curl UMETA(DisplayName = "Curl (Vector->Vector)"),
+	Magnitude UMETA(DisplayName = "Magnitude (Vector->Scalar)"),
+	Normalize UMETA(DisplayName = "Normalize (Vector->Vector)"),
+};
+
+UENUM(BlueprintType)
+enum class EDataflowVolumeAnalysisOutputName : uint8
+{
+	AppendOperationName UMETA(DisplayName = "Append Operation Name"),
+	CustomName UMETA(DisplayName = "Custom Name"),
+};
\ No newline at end of file
diff --git a/Engine/Plugins/Experimental/DataflowVolumetricPlugin/Source/DataflowVolumeCore/Private/Dataflow/DataflowBoolVolume.cpp b/Engine/Plugins/Experimental/DataflowVolumetricPlugin/Source/DataflowVolumeCore/Private/Dataflow/DataflowBoolVolume.cpp
new file mode 100644
--- /dev/null
+++ b/Engine/Plugins/Experimental/DataflowVolumetricPlugin/Source/DataflowVolumeCore/Private/Dataflow/DataflowBoolVolume.cpp
@@ -0,0 +1,49 @@
+// Copyright Epic Games, Inc. All Rights Reserved.
+
+#include "Dataflow/DataflowBoolVolume.h"
+#include "Dataflow/DataflowBoolVolumeImpl.h"
+
+/* --------------------------------------------------------------------------------------------------------------- */
+/* FDataflowBoolVolume Implementation */
+/* --------------------------------------------------------------------------------------------------------------- */
+
+FDataflowBoolVolume::FDataflowBoolVolume(TSharedPtr<UE::DataflowVolume::Private::FDataflowBoolVolumeImpl> InBoolVolume)
+	: FDataflowVolume(InBoolVolume, FDataflowBoolVolume::TypeName)
+{
+	static_assert(sizeof(FDataflowVolume) == sizeof(*this));
+}
+
+TSharedPtr<const UE::DataflowVolume::Private::FDataflowBoolVolumeImpl> FDataflowBoolVolume::GetBoolVolumeImpl() const
+{
+	return StaticCastSharedPtr<const UE::DataflowVolume::Private::FDataflowBoolVolumeImpl>(Volume);
+}
+
+TSharedPtr<UE::DataflowVolume::Private::FDataflowBoolVolumeImpl> FDataflowBoolVolume::GetBoolVolumeImpl()
+{
+	return StaticCastSharedPtr<UE::DataflowVolume::Private::FDataflowBoolVolumeImpl>(Volume);
+}
+
+FVector FDataflowBoolVolume::GetVoxelSize() const
+{
+	using namespace UE::DataflowVolume::Private;
+
+	if (TSharedPtr<const FDataflowBoolVolumeImpl> BoolVolumeImpl = GetBoolVolumeImpl())
+	{
+		return BoolVolumeImpl->GetVoxelSize();
+	}
+
+	return {};
+}
+
+FBox FDataflowBoolVolume::GetVolumeBoundingBox() const
+{
+	using namespace UE::DataflowVolume::Private;
+
+	if (TSharedPtr<const FDataflowBoolVolumeImpl> BoolVolumeImpl = GetBoolVolumeImpl())
+	{
+		return BoolVolumeImpl->GetActiveVoxelBoundingBox();
+	}
+
+	return {};
+}
+
diff --git a/Engine/Plugins/Experimental/DataflowVolumetricPlugin/Source/DataflowVolumeCore/Private/Dataflow/DataflowBoolVolumeImpl.cpp b/Engine/Plugins/Experimental/DataflowVolumetricPlugin/Source/DataflowVolumeCore/Private/Dataflow/DataflowBoolVolumeImpl.cpp
new file mode 100644
--- /dev/null
+++ b/Engine/Plugins/Experimental/DataflowVolumetricPlugin/Source/DataflowVolumeCore/Private/Dataflow/DataflowBoolVolumeImpl.cpp
@@ -0,0 +1,83 @@
+// Copyright Epic Games, Inc. All Rights Reserved.
+
+#include "Dataflow/DataflowBoolVolumeImpl.h"
+
+namespace UE::DataflowVolume::Private
+{
+	/* --------------------------------------------------------------------------------------------------------------- */
+	/* FDataflowBoolVolumeImpl */
+	/* --------------------------------------------------------------------------------------------------------------- */
+	void FDataflowBoolVolumeImpl::SetGrid(const openvdb::BoolGrid::Ptr& InBoolGrid, const FString& InGridName)
+	{
+		FDataflowVolumeImpl::SetGrid(InBoolGrid, InGridName);
+
+		BoolGrid = InBoolGrid;
+	}
+
+	openvdb::BoolGrid::Ptr FDataflowBoolVolumeImpl::GetGrid() const
+	{
+		return BoolGrid;
+	}
+
+	FString FDataflowBoolVolumeImpl::GetGridType() const
+	{
+		if (BoolGrid)
+		{
+			std::string GridType = BoolGrid->gridType();
+			return FString(GridType.c_str());
+		}
+
+		return {};
+	}
+
+	FString FDataflowBoolVolumeImpl::VolumeInfo() const
+	{
+		FString BaseString = FDataflowVolumeImpl::VolumeInfo();
+
+		FString OutputStr;
+		OutputStr.Appendf(TEXT("\n----------------------------------------\n"));
+		OutputStr += BaseString;
+
+		FString GridType = this->GetGridType();
+
+		OutputStr.Appendf(TEXT("Type: %s\n"),
+			*GridType);
+
+		OutputStr.Appendf(TEXT("----------------------------------------\n"));
+
+		return OutputStr;
+	}
+
+	void FDataflowBoolVolumeImpl::GetActiveVoxels(TArray<FBox>& OutActiveVoxels) const
+	{
+		if (BoolGrid)
+		{
+			const openvdb::Index64 ActiveVoxelCount = BoolGrid->activeVoxelCount();
+			OutActiveVoxels.AddUninitialized(ActiveVoxelCount);
+
+			openvdb::math::Transform T = this->BoolGrid->transform();
+
+			int32 Idx = 0;
+			for (openvdb::BoolGrid::ValueOnCIter iter = this->BoolGrid->cbeginValueOn(); iter.test(); ++iter)
+			{
+				if (iter.isVoxelValue())
+				{
+					openvdb::CoordBBox BBox;
+					iter.getBoundingBox(BBox);
+
+					const openvdb::Vec3d BBoxMin(BBox.min().x() - 0.5, BBox.min().y() - 0.5, BBox.min().z() - 0.5);
+					const openvdb::Vec3d BBoxMax(BBox.max().x() + 0.5, BBox.max().y() + 0.5, BBox.max().z() + 0.5);
+
+					const openvdb::Vec3d BBoxMinWS = BoolGrid->transform().indexToWorld(BBoxMin);
+					const openvdb::Vec3d BBoxMaxWS = BoolGrid->transform().indexToWorld(BBoxMax);
+
+					FVector Min(BBoxMinWS.x(), BBoxMinWS.y(), BBoxMinWS.z());
+					FVector Max(BBoxMaxWS.x(), BBoxMaxWS.y(), BBoxMaxWS.z());
+
+					OutActiveVoxels[Idx++] = FBox(Min, Max);
+				}
+			}
+		}
+	}
+}
+
diff --git a/Engine/Plugins/Experimental/DataflowVolumetricPlugin/Source/DataflowVolumeCore/Private/Dataflow/DataflowBoolVolumeImpl.h b/Engine/Plugins/Experimental/DataflowVolumetricPlugin/Source/DataflowVolumeCore/Private/Dataflow/DataflowBoolVolumeImpl.h
new file mode 100644
--- /dev/null
+++ b/Engine/Plugins/Experimental/DataflowVolumetricPlugin/Source/DataflowVolumeCore/Private/Dataflow/DataflowBoolVolumeImpl.h
@@ -0,0 +1,28 @@
+// Copyright Epic Games, Inc. All Rights Reserved.
+
+#pragma once
+
+#include "Dataflow/DataflowVolumeImpl.h"
+#include "Dataflow/OpenVDB.h"
+
+namespace UE::DataflowVolume::Private
+{
+	/* --------------------------------------------------------------------------------------------------------------- */
+	/* FDataflowBoolVolumeImpl */
+	/* --------------------------------------------------------------------------------------------------------------- */
+	class FDataflowBoolVolumeImpl : public FDataflowVolumeImpl
+	{
+	private:
+		openvdb::BoolGrid::Ptr BoolGrid = nullptr;
+
+	public:
+		void SetGrid(const openvdb::BoolGrid::Ptr& InBoolGrid, const FString& InGridName = "");
+		openvdb::BoolGrid::Ptr GetGrid() const;
+		FString GetGridType() const;
+		FString VolumeInfo() const;
+		virtual void GetActiveVoxels(TArray<FBox>& OutActiveVoxels) const override;
+	};
+}
+
+
+
diff --git a/Engine/Plugins/Experimental/DataflowVolumetricPlugin/Source/DataflowVolumeCore/Private/Dataflow/DataflowFloatVectorVolume.cpp b/Engine/Plugins/Experimental/DataflowVolumetricPlugin/Source/DataflowVolumeCore/Private/Dataflow/DataflowFloatVectorVolume.cpp
new file mode 100644
--- /dev/null
+++ b/Engine/Plugins/Experimental/DataflowVolumetricPlugin/Source/DataflowVolumeCore/Private/Dataflow/DataflowFloatVectorVolume.cpp
@@ -0,0 +1,156 @@
+// Copyright Epic Games, Inc. All Rights Reserved.
+
+#include "Dataflow/DataflowFloatVectorVolume.h"
+#include "Dataflow/DataflowFloatVolume.h"
+#include "Dataflow/DataflowBoolVolume.h"
+#include "Dataflow/DataflowFloatVolumeImpl.h"
+#include "Dataflow/DataflowBoolVolumeImpl.h"
+#include "Dataflow/DataflowFloatVectorVolumeImpl.h"
+
+/* --------------------------------------------------------------------------------------------------------------- */
+/* FDataflowFloatVectorVolume Implementation */
+/* --------------------------------------------------------------------------------------------------------------- */
+
+FDataflowFloatVectorVolume::FDataflowFloatVectorVolume(TSharedPtr<UE::DataflowVolume::Private::FDataflowFloatVectorVolumeImpl> InFloatVectorVolume)
+	: FDataflowVolume(InFloatVectorVolume, FDataflowFloatVectorVolume::TypeName)
+{
+	static_assert(sizeof(FDataflowVolume) == sizeof(*this));
+}
+
+/* --------------------------------------------------------------------------------------------------------------- */
+
+TSharedPtr<const UE::DataflowVolume::Private::FDataflowFloatVectorVolumeImpl> FDataflowFloatVectorVolume::GetFloatVectorVolumeImpl() const
+{
+	return StaticCastSharedPtr<const UE::DataflowVolume::Private::FDataflowFloatVectorVolumeImpl>(Volume);
+}
+
+/* --------------------------------------------------------------------------------------------------------------- */
+
+TSharedPtr<UE::DataflowVolume::Private::FDataflowFloatVectorVolumeImpl> FDataflowFloatVectorVolume::GetFloatVectorVolumeImpl()
+{
+	return StaticCastSharedPtr<UE::DataflowVolume::Private::FDataflowFloatVectorVolumeImpl>(Volume);
+}
+
+/* --------------------------------------------------------------------------------------------------------------- */
+
+FString FDataflowFloatVectorVolume::GetVectorType() const
+{
+	return GetFloatVectorVolumeImpl()->GetVectorType();
+}
+
+/* --------------------------------------------------------------------------------------------------------------- */
+
+void FDataflowFloatVectorVolume::VolumeSample(const TArray<FVector>& InPoints, TArray<FVector>& OutValues) const
+{
+	using namespace UE::DataflowVolume::Private;
+
+	if (TSharedPtr<const FDataflowFloatVectorVolumeImpl> FloatVectorVolumeImpl = GetFloatVectorVolumeImpl())
+	{
+		FloatVectorVolumeImpl->VolumeSample(InPoints, OutValues);
+	}
+}
+
+/* --------------------------------------------------------------------------------------------------------------- */
+
+FVector FDataflowFloatVectorVolume::GetVoxelSize() const
+{
+	using namespace UE::DataflowVolume::Private;
+
+	if (TSharedPtr<const FDataflowFloatVectorVolumeImpl> FloatVectorVolumeImpl = GetFloatVectorVolumeImpl())
+	{
+		return FloatVectorVolumeImpl->GetVoxelSize();
+	}
+
+	return {};
+}
+
+/* --------------------------------------------------------------------------------------------------------------- */
+
+FBox FDataflowFloatVectorVolume::GetVolumeBoundingBox() const
+{
+	using namespace UE::DataflowVolume::Private;
+
+	if (TSharedPtr<const FDataflowFloatVectorVolumeImpl> FloatVectorVolumeImpl = GetFloatVectorVolumeImpl())
+	{
+		return FloatVectorVolumeImpl->GetActiveVoxelBoundingBox();
+	}
+
+	return {};
+}
+
+/* --------------------------------------------------------------------------------------------------------------- */
+
+FDataflowFloatVolume FDataflowFloatVectorVolume::ComputeDivergence(const EDataflowVolumeAnalysisOutputName OutputName, const FString& CustomName, const FDataflowBoolVolume* MaskVolume) const
+{
+	using namespace UE::DataflowVolume::Private;
+
+	if (TSharedPtr<const FDataflowFloatVectorVolumeImpl> FloatVectorVolumeImpl = GetFloatVectorVolumeImpl())
+	{
+		if (TSharedPtr<const FDataflowBoolVolumeImpl> MaskVolumeImpl = MaskVolume ? MaskVolume->GetBoolVolumeImpl() : MakeShared<FDataflowBoolVolumeImpl>())
+		{
+			TSharedPtr<FDataflowFloatVolumeImpl> DivergenceVolume = FloatVectorVolumeImpl->ComputeDivergence(OutputName, CustomName, MaskVolumeImpl);
+
+			return FDataflowFloatVolume(DivergenceVolume);
+		}
+	}
+
+	return FDataflowFloatVolume();
+}
+
+/* --------------------------------------------------------------------------------------------------------------- */
+
+FDataflowFloatVolume FDataflowFloatVectorVolume::ComputeMagnitude(const EDataflowVolumeAnalysisOutputName OutputName, const FString& CustomName, const FDataflowBoolVolume* MaskVolume) const
+{
+	using namespace UE::DataflowVolume::Private;
+
+	if (TSharedPtr<const FDataflowFloatVectorVolumeImpl> FloatVectorVolumeImpl = GetFloatVectorVolumeImpl())
+	{
+		if (TSharedPtr<const FDataflowBoolVolumeImpl> MaskVolumeImpl = MaskVolume ? MaskVolume->GetBoolVolumeImpl() : MakeShared<FDataflowBoolVolumeImpl>())
+		{
+			TSharedPtr<FDataflowFloatVolumeImpl> LengthVolume = FloatVectorVolumeImpl->ComputeMagnitude(OutputName, CustomName, MaskVolumeImpl);
+			return FDataflowFloatVolume(LengthVolume);
+		}
+	}
+
+	return FDataflowFloatVolume();
+}
+
+/* --------------------------------------------------------------------------------------------------------------- */
+
+FDataflowFloatVectorVolume FDataflowFloatVectorVolume::ComputeCurl(const EDataflowVolumeAnalysisOutputName OutputName, const FString& CustomName, const FDataflowBoolVolume* MaskVolume) const
+{
+	using namespace UE::DataflowVolume::Private;
+
+	if (TSharedPtr<const FDataflowFloatVectorVolumeImpl> FloatVectorVolumeImpl = GetFloatVectorVolumeImpl())
+	{
+		if (TSharedPtr<const FDataflowBoolVolumeImpl> MaskVolumeImpl = MaskVolume ? MaskVolume->GetBoolVolumeImpl() : MakeShared<FDataflowBoolVolumeImpl>())
+		{
+			TSharedPtr<FDataflowFloatVectorVolumeImpl> CurlVolume = FloatVectorVolumeImpl->ComputeCurl(OutputName, CustomName, MaskVolumeImpl);
+			return FDataflowFloatVectorVolume(CurlVolume);
+		}
+	}
+
+	return FDataflowFloatVectorVolume();
+}
+
+/* --------------------------------------------------------------------------------------------------------------- */
+
+FDataflowFloatVectorVolume FDataflowFloatVectorVolume::ComputeNormalize(const EDataflowVolumeAnalysisOutputName OutputName, const FString& CustomName, const FDataflowBoolVolume* MaskVolume) const
+{
+	using namespace UE::DataflowVolume::Private;
+
+	if (TSharedPtr<const FDataflowFloatVectorVolumeImpl> FloatVectorVolumeImpl = GetFloatVectorVolumeImpl())
+	{
+		if (TSharedPtr<const FDataflowBoolVolumeImpl> MaskVolumeImpl = MaskVolume ? MaskVolume->GetBoolVolumeImpl() : MakeShared<FDataflowBoolVolumeImpl>())
+		{
+			TSharedPtr<FDataflowFloatVectorVolumeImpl> NormalVolume = FloatVectorVolumeImpl->ComputeNormalize(OutputName, CustomName, MaskVolumeImpl);
+			return FDataflowFloatVectorVolume(NormalVolume);
+		}
+	}
+
+	return FDataflowFloatVectorVolume();
+}
+
+/* --------------------------------------------------------------------------------------------------------------- */
+
+
diff --git a/Engine/Plugins/Experimental/DataflowVolumetricPlugin/Source/DataflowVolumeCore/Private/Dataflow/DataflowFloatVectorVolumeImpl.cpp b/Engine/Plugins/Experimental/DataflowVolumetricPlugin/Source/DataflowVolumeCore/Private/Dataflow/DataflowFloatVectorVolumeImpl.cpp
new file mode 100644
--- /dev/null
+++ b/Engine/Plugins/Experimental/DataflowVolumetricPlugin/Source/DataflowVolumeCore/Private/Dataflow/DataflowFloatVectorVolumeImpl.cpp
@@ -0,0 +1,404 @@
+// Copyright Epic Games, Inc. All Rights Reserved.
+
+#include "Dataflow/DataflowFloatVectorVolumeImpl.h"
+#include "Dataflow/DataflowBoolVolumeImpl.h"
+#include <Dataflow/DataflowBoolVolume.h>
+
+namespace UE::DataflowVolume::Private
+{
+	/* --------------------------------------------------------------------------------------------------------------- */
+	/* FDataflowFloatVectorVolumeImpl */
+	/* --------------------------------------------------------------------------------------------------------------- */
+	void FDataflowFloatVectorVolumeImpl::SetGrid(const openvdb::VectorGrid::Ptr& InFloatVectorGrid, const openvdb::VecType InVecType, const FString& InGridName)
+	{
+		FDataflowVolumeImpl::SetGrid(InFloatVectorGrid, InGridName);
+
+		FloatVectorGrid = InFloatVectorGrid;
+		FloatVectorGrid->setVectorType(InVecType);
+	}
+
+	/* --------------------------------------------------------------------------------------------------------------- */
+
+	openvdb::VectorGrid::Ptr FDataflowFloatVectorVolumeImpl::GetGrid() const
+	{
+		return FloatVectorGrid;
+	}
+
+	/* --------------------------------------------------------------------------------------------------------------- */
+
+	FString FDataflowFloatVectorVolumeImpl::GetGridType() const
+	{
+		if (FloatVectorGrid)
+		{
+			std::string GridType = FloatVectorGrid->gridType();
+			return FString(GridType.c_str());
+		}
+
+		return {};
+	}
+
+	/* --------------------------------------------------------------------------------------------------------------- */
+
+	/// The type of a vector determines how transforms are applied to it:
+	/// 
+	/// Invariant
+	/// Does not transform (e.g., tuple, uvw, color)
+	///
+	/// Covariant
+	/// Apply inverse-transpose transformation: @e w = 0, ignores translation
+	///     (e.g., gradient/normal)
+	///
+	/// Covariant Normalize
+	/// Apply inverse-transpose transformation: @e w = 0, ignores translation,
+	///     vectors are renormalized (e.g., unit normal)
+	///
+	/// Contravariant Relative
+	/// Apply "regular" transformation: @e w = 0, ignores translation
+	///     (e.g., displacement, velocity, acceleration)
+	///
+	/// Contravariant Absolute
+	/// Apply "regular" transformation: @e w = 1, vector translates (e.g., position)
+	/// 
+	FString FDataflowFloatVectorVolumeImpl::GetVectorType() const
+	{
+		if (FloatVectorGrid)
+		{
+			FString VectorType;
+			switch (FloatVectorGrid->getVectorType()) {
+			case openvdb::VEC_INVARIANT:
+				VectorType = TEXT("Invariant");
+				break;
+			case openvdb::VEC_COVARIANT:
+				VectorType = TEXT("Covariant");
+				break;
+			case openvdb::VEC_COVARIANT_NORMALIZE:
+				VectorType = TEXT("Covariant Normalize");
+				break;
+			case openvdb::VEC_CONTRAVARIANT_RELATIVE:
+				VectorType = TEXT("Contravariant Relative");
+				break;
+			case openvdb::VEC_CONTRAVARIANT_ABSOLUTE:
+				VectorType = TEXT("Contravariant Absolute");
+				break;
+
+			default:
+				break;
+			}
+
+			return VectorType;
+		}
+
+		return {};
+	}
+
+	/* --------------------------------------------------------------------------------------------------------------- */
+
+	FString FDataflowFloatVectorVolumeImpl::VolumeInfo() const
+	{
+		FString BaseString = FDataflowVolumeImpl::VolumeInfo();
+
+		FString OutputStr;
+		OutputStr.Appendf(TEXT("\n----------------------------------------\n"));
+		OutputStr += BaseString;
+
+		FString GridType = this->GetGridType();
+		FString VectorType = this->GetVectorType();
+
+		OutputStr.Appendf(TEXT("Type: %s\nVector Type: %s\n"),
+			*GridType,
+			*VectorType);
+
+		OutputStr.Appendf(TEXT("----------------------------------------\n"));
+
+		return OutputStr;
+	}
+
+	/* --------------------------------------------------------------------------------------------------------------- */
+
+	void FDataflowFloatVectorVolumeImpl::VolumeSample(
+		const TArray<FVector>& InPoints,
+		TArray<FVector>& OutValues) const
+	{
+		using namespace openvdb;
+
+		if (this->FloatVectorGrid->empty())
+		{
+			return;
+		}
+
+		openvdb::VectorGrid::ConstAccessor Accessor = this->FloatVectorGrid->getConstAccessor();
+
+		openvdb::tools::GridSampler<openvdb::VectorGrid::ConstAccessor, openvdb::tools::BoxSampler>
+			fastSampler(Accessor, this->FloatVectorGrid->transform());
+
+		int32 Idx = 0;
+		for (const FVector& Point : InPoints)
+		{
+			openvdb::Vec3f SampledValue = fastSampler.wsSample(openvdb::Vec3R(Point.X, Point.Y, Point.Z));
+			OutValues[Idx++] = FVector(SampledValue.x(), SampledValue.y(), SampledValue.z());
+		}
+
+		return;
+	}
+
+	/* --------------------------------------------------------------------------------------------------------------- */
+
+	void FDataflowFloatVectorVolumeImpl::GetActiveVoxels(TArray<FBox>& OutActiveVoxels) const
+	{
+		const openvdb::Index64 ActiveVoxelCount = FloatVectorGrid->activeVoxelCount();
+		OutActiveVoxels.AddUninitialized(ActiveVoxelCount);
+
+		openvdb::math::Transform T = this->FloatVectorGrid->transform();
+
+		int32 Idx = 0;
+		for (openvdb::VectorGrid::ValueOnCIter iter = this->FloatVectorGrid->cbeginValueOn(); iter.test(); ++iter)
+		{
+			if (iter.isVoxelValue())
+			{
+				openvdb::CoordBBox BBox;
+				iter.getBoundingBox(BBox);
+
+				const openvdb::Vec3d BBoxMin(BBox.min().x() - 0.5, BBox.min().y() - 0.5, BBox.min().z() - 0.5);
+				const openvdb::Vec3d BBoxMax(BBox.max().x() + 0.5, BBox.max().y() + 0.5, BBox.max().z() + 0.5);
+
+				const openvdb::Vec3d BBoxMinWS = FloatVectorGrid->transform().indexToWorld(BBoxMin);
+				const openvdb::Vec3d BBoxMaxWS = FloatVectorGrid->transform().indexToWorld(BBoxMax);
+
+				FVector Min(BBoxMinWS.x(), BBoxMinWS.y(), BBoxMinWS.z());
+				FVector Max(BBoxMaxWS.x(), BBoxMaxWS.y(), BBoxMaxWS.z());
+
+				OutActiveVoxels[Idx++] = FBox(Min, Max);
+			}
+		}
+	}
+
+	/* --------------------------------------------------------------------------------------------------------------- */
+
+	TSharedPtr<FDataflowFloatVolumeImpl> FDataflowFloatVectorVolumeImpl::ComputeDivergence(const EDataflowVolumeAnalysisOutputName OutputName, const FString& CustomName, TSharedPtr<const FDataflowBoolVolumeImpl> InMaskVolumeImpl) const
+	{
+		using namespace openvdb;
+
+		constexpr openvdb::GridClass GridClass = openvdb::GRID_UNKNOWN;
+		const FString GridName = "Divergence";
+
+		TSharedPtr<FDataflowFloatVolumeImpl> FloatVolumeImpl = MakeShared<FDataflowFloatVolumeImpl>();
+
+		if (this->FloatVectorGrid->empty())
+		{
+			FloatVolumeImpl->SetGrid(openvdb::FloatGrid::Ptr(), GridClass, GridName);
+			return FloatVolumeImpl;
+		}
+
+		FString NewGridName;
+		if (OutputName == EDataflowVolumeAnalysisOutputName::AppendOperationName)
+		{
+			const FString CurrentGridName(this->FloatVectorGrid->getName().c_str());
+			NewGridName = CurrentGridName + "_" + GridName;
+		}
+		else if (OutputName == EDataflowVolumeAnalysisOutputName::CustomName)
+		{
+			NewGridName = CustomName;
+		}
+
+		// TODO: Implement this
+//			UE::DataflowVolumeUtils::FInterrupter NullInterrupter;
+
+		openvdb::FloatGrid::Ptr DivergenceGrid;
+
+		if (InMaskVolumeImpl->GetGrid())
+		{
+			openvdb::BoolGrid RegionMask;
+
+			RegionMask.setTransform(this->FloatVectorGrid->transform().copy());
+
+			openvdb::tools::resampleToMatch<openvdb::tools::PointSampler>(*InMaskVolumeImpl->GetGrid().get(), RegionMask);
+
+			openvdb::tools::Divergence<openvdb::VectorGrid, openvdb::BoolGrid> DivergenceTool(*this->FloatVectorGrid.get(), RegionMask);
+			DivergenceGrid = DivergenceTool.process(/*threaded*/true);
+		}
+		else
+		{
+			openvdb::tools::Divergence<openvdb::VectorGrid> DivergenceTool(*this->FloatVectorGrid.get());
+			DivergenceGrid = DivergenceTool.process(/*threaded*/true);
+		}
+
+		FloatVolumeImpl->SetGrid(DivergenceGrid, GridClass, NewGridName);
+
+		return FloatVolumeImpl;
+	}
+
+	/* --------------------------------------------------------------------------------------------------------------- */
+
+	TSharedPtr<FDataflowFloatVolumeImpl> FDataflowFloatVectorVolumeImpl::ComputeMagnitude(const EDataflowVolumeAnalysisOutputName OutputName, const FString& CustomName, TSharedPtr<const FDataflowBoolVolumeImpl> InMaskVolumeImpl) const
+	{
+		using namespace openvdb;
+
+		constexpr openvdb::GridClass GridClass = openvdb::GRID_UNKNOWN;
+		const FString GridName = "Magnitude";
+
+		TSharedPtr<FDataflowFloatVolumeImpl> FloatVolumeImpl = MakeShared<FDataflowFloatVolumeImpl>();
+
+		if (this->FloatVectorGrid->empty())
+		{
+			FloatVolumeImpl->SetGrid(openvdb::FloatGrid::Ptr(), GridClass, GridName);
+			return FloatVolumeImpl;
+		}
+
+		FString NewGridName;
+		if (OutputName == EDataflowVolumeAnalysisOutputName::AppendOperationName)
+		{
+			const FString CurrentGridName(this->FloatVectorGrid->getName().c_str());
+			NewGridName = CurrentGridName + "_" + GridName;
+		}
+		else if (OutputName == EDataflowVolumeAnalysisOutputName::CustomName)
+		{
+			NewGridName = CustomName;
+		}
+
+		// TODO: Implement this
+	//			UE::DataflowVolumeUtils::FInterrupter NullInterrupter;
+
+		openvdb::FloatGrid::Ptr MagnitudeGrid;
+		if (InMaskVolumeImpl->GetGrid())
+		{
+			openvdb::BoolGrid RegionMask;
+
+			RegionMask.setTransform(this->FloatVectorGrid->transform().copy());
+
+			openvdb::tools::resampleToMatch<openvdb::tools::PointSampler>(*InMaskVolumeImpl->GetGrid().get(), RegionMask);
+
+			openvdb::tools::Magnitude <openvdb::VectorGrid, openvdb::BoolGrid> MagnitudeTool(*this->FloatVectorGrid.get(), RegionMask);
+			MagnitudeGrid = MagnitudeTool.process(/*threaded*/true);
+		}
+		else
+		{
+			openvdb::tools::Magnitude <openvdb::VectorGrid> MagnitudeTool(*this->FloatVectorGrid.get());
+			MagnitudeGrid = MagnitudeTool.process(/*threaded*/true);
+		}
+
+		FloatVolumeImpl->SetGrid(MagnitudeGrid, GridClass, NewGridName);
+
+		return FloatVolumeImpl;
+	}
+
+	/* --------------------------------------------------------------------------------------------------------------- */
+
+	TSharedPtr<FDataflowFloatVectorVolumeImpl> FDataflowFloatVectorVolumeImpl::ComputeCurl(const EDataflowVolumeAnalysisOutputName OutputName, const FString& CustomName, TSharedPtr<const FDataflowBoolVolumeImpl> InMaskVolumeImpl) const
+	{
+		using namespace openvdb;
+
+		constexpr openvdb::VecType VectorGridType = openvdb::VEC_COVARIANT;
+		const FString GridName = "Curl";
+
+		TSharedPtr<FDataflowFloatVectorVolumeImpl> FloatVectorVolumeImpl = MakeShared<FDataflowFloatVectorVolumeImpl>();
+
+		if (this->FloatVectorGrid->empty())
+		{
+			FloatVectorVolumeImpl->SetGrid(openvdb::VectorGrid::Ptr(), VectorGridType, GridName);
+			return FloatVectorVolumeImpl;
+		}
+
+		FString NewGridName;
+		if (OutputName == EDataflowVolumeAnalysisOutputName::AppendOperationName)
+		{
+			const FString CurrentGridName(this->FloatVectorGrid->getName().c_str());
+			NewGridName = CurrentGridName + "_" + GridName;
+		}
+		else if (OutputName == EDataflowVolumeAnalysisOutputName::CustomName)
+		{
+			NewGridName = CustomName;
+		}
+
+		// TODO: Implement this
+		//			UE::DataflowVolumeUtils::FInterrupter NullInterrupter;
+
+		openvdb::VectorGrid::Ptr CurlVectorGrid;
+
+		if (InMaskVolumeImpl->GetGrid())
+		{
+			openvdb::BoolGrid RegionMask;
+
+			RegionMask.setTransform(this->FloatVectorGrid->transform().copy());
+
+			openvdb::tools::resampleToMatch<openvdb::tools::PointSampler>(*InMaskVolumeImpl->GetGrid().get(), RegionMask);
+
+			openvdb::tools::Curl<openvdb::VectorGrid, openvdb::BoolGrid> CurlTool(*this->FloatVectorGrid.get(), RegionMask);
+			CurlVectorGrid = CurlTool.process(/*threaded*/true);
+		}
+		else
+		{
+			openvdb::tools::Curl<openvdb::VectorGrid> CurlTool(*this->FloatVectorGrid.get());
+			CurlVectorGrid = CurlTool.process(/*threaded*/true);
+		}
+
+		FloatVectorVolumeImpl->SetGrid(CurlVectorGrid, VectorGridType, NewGridName);
+
+		return FloatVectorVolumeImpl;
+	}
+
+	/* --------------------------------------------------------------------------------------------------------------- */
+
+	TSharedPtr<FDataflowFloatVectorVolumeImpl> FDataflowFloatVectorVolumeImpl::ComputeNormalize(const EDataflowVolumeAnalysisOutputName OutputName, const FString& CustomName, TSharedPtr<const FDataflowBoolVolumeImpl> InMaskVolumeImpl) const
+	{
+		using namespace openvdb;
+
+		constexpr openvdb::VecType VectorGridType = openvdb::VEC_COVARIANT;
+		const FString GridName = "Normalize";
+
+		TSharedPtr<FDataflowFloatVectorVolumeImpl> FloatVectorVolumeImpl = MakeShared<FDataflowFloatVectorVolumeImpl>();
+
+		if (this->FloatVectorGrid->empty())
+		{
+			FloatVectorVolumeImpl->SetGrid(openvdb::VectorGrid::Ptr(), VectorGridType, GridName);
+			return FloatVectorVolumeImpl;
+		}
+
+		FString NewGridName;
+		if (OutputName == EDataflowVolumeAnalysisOutputName::AppendOperationName)
+		{
+			const FString CurrentGridName(this->FloatVectorGrid->getName().c_str());
+			NewGridName = CurrentGridName + "_" + "GridName";
+		}
+		else if (OutputName == EDataflowVolumeAnalysisOutputName::CustomName)
+		{
+			NewGridName = CustomName;
+		}
+
+		// TODO: Implement this
+		//			UE::DataflowVolumeUtils::FInterrupter NullInterrupter;
+
+		openvdb::VectorGrid::Ptr NormalVectorGrid;
+		if (InMaskVolumeImpl->GetGrid())
+		{
+			openvdb::BoolGrid RegionMask;
+
+			RegionMask.setTransform(this->FloatVectorGrid->transform().copy());
+
+			openvdb::tools::resampleToMatch<openvdb::tools::PointSampler>(*InMaskVolumeImpl->GetGrid().get(), RegionMask);
+
+			openvdb::tools::Normalize<openvdb::VectorGrid, openvdb::BoolGrid> NormalizeTool(*this->FloatVectorGrid.get(), RegionMask);
+			NormalVectorGrid = NormalizeTool.process(/*threaded*/true);
+		}
+		else
+		{
+			openvdb::tools::Normalize<openvdb::VectorGrid> NormalizeTool(*this->FloatVectorGrid.get());
+			NormalVectorGrid = NormalizeTool.process(/*threaded*/true);
+		}
+
+		FloatVectorVolumeImpl->SetGrid(NormalVectorGrid, VectorGridType, NewGridName);
+
+		return FloatVectorVolumeImpl;
+	}
+
+	/* --------------------------------------------------------------------------------------------------------------- */
+
+
+
+
+
+
+	
+
+
+
+}
\ No newline at end of file
diff --git a/Engine/Plugins/Experimental/DataflowVolumetricPlugin/Source/DataflowVolumeCore/Private/Dataflow/DataflowFloatVectorVolumeImpl.h b/Engine/Plugins/Experimental/DataflowVolumetricPlugin/Source/DataflowVolumeCore/Private/Dataflow/DataflowFloatVectorVolumeImpl.h
new file mode 100644
--- /dev/null
+++ b/Engine/Plugins/Experimental/DataflowVolumetricPlugin/Source/DataflowVolumeCore/Private/Dataflow/DataflowFloatVectorVolumeImpl.h
@@ -0,0 +1,44 @@
+// Copyright Epic Games, Inc. All Rights Reserved.
+
+#pragma once
+
+#include "Dataflow/DataflowVolumeImpl.h"
+#include "Dataflow/DataflowFloatVolumeImpl.h"
+#include "Dataflow/DataflowBoolVolumeImpl.h"
+#include "Dataflow/OpenVDB.h"
+#include "Dataflow/DataflowVolumeNodeEnums.h"
+
+namespace UE::DataflowVolume::Private
+{
+	class FDataflowFloatVolumeImpl;
+	class FDataflowBoolVolumeImpl;
+
+	/* --------------------------------------------------------------------------------------------------------------- */
+	/* FDataflowFloatVectorVolumeImpl */
+	/* --------------------------------------------------------------------------------------------------------------- */
+	class FDataflowFloatVectorVolumeImpl : public FDataflowVolumeImpl
+	{
+	private:
+		openvdb::VectorGrid::Ptr FloatVectorGrid = nullptr;
+
+	public:
+		void SetGrid(const openvdb::VectorGrid::Ptr& InFloatVectorGrid, const openvdb::VecType InVecType, const FString& InGridName = "Gradient");
+
+		openvdb::VectorGrid::Ptr GetGrid() const;
+
+		FString GetGridType() const;
+
+		FString GetVectorType() const;
+
+		FString VolumeInfo() const;
+
+		void VolumeSample(const TArray<FVector>& InPoints, TArray<FVector>& OutValues) const;
+
+		virtual void GetActiveVoxels(TArray<FBox>& OutActiveVoxels) const override;
+
+		TSharedPtr<FDataflowFloatVolumeImpl> ComputeDivergence(const EDataflowVolumeAnalysisOutputName OutputName, const FString& CustomName, TSharedPtr<const FDataflowBoolVolumeImpl> InMaskVolumeImpl) const;
+		TSharedPtr<FDataflowFloatVolumeImpl> ComputeMagnitude(const EDataflowVolumeAnalysisOutputName OutputName, const FString& CustomName, TSharedPtr<const FDataflowBoolVolumeImpl> InMaskVolumeImpl) const;
+		TSharedPtr<FDataflowFloatVectorVolumeImpl> ComputeCurl(const EDataflowVolumeAnalysisOutputName OutputName, const FString& CustomName, TSharedPtr<const FDataflowBoolVolumeImpl> InMaskVolumeImpl) const;
+		TSharedPtr<FDataflowFloatVectorVolumeImpl> ComputeNormalize(const EDataflowVolumeAnalysisOutputName OutputName, const FString& CustomName, TSharedPtr<const FDataflowBoolVolumeImpl> InMaskVolumeImpl) const;
+	};
+}
diff --git a/Engine/Plugins/Experimental/DataflowVolumetricPlugin/Source/DataflowVolumeCore/Private/Dataflow/DataflowFloatVolume.cpp b/Engine/Plugins/Experimental/DataflowVolumetricPlugin/Source/DataflowVolumeCore/Private/Dataflow/DataflowFloatVolume.cpp
new file mode 100644
--- /dev/null
+++ b/Engine/Plugins/Experimental/DataflowVolumetricPlugin/Source/DataflowVolumeCore/Private/Dataflow/DataflowFloatVolume.cpp
@@ -0,0 +1,435 @@
+// Copyright Epic Games, Inc. All Rights Reserved.
+
+#include "Dataflow/DataflowFloatVolume.h"
+#include "Dataflow/DataflowIntVolume.h"
+#include "Dataflow/DataflowBoolVolume.h"
+#include "Dataflow/DataflowFloatVolumeImpl.h"
+#include "Dataflow/DataflowFloatVectorVolume.h"
+#include "Dataflow/DataflowVolumeUtils.h"
+#include "MeshDescription.h"
+#include "MeshAttributes.h"
+#include "MeshAttributeArray.h"
+#include "Dataflow/OpenVDB.h"
+
+/* --------------------------------------------------------------------------------------------------------------- */
+/* FDataflowFloatVolume Implementation */
+/* --------------------------------------------------------------------------------------------------------------- */
+
+FDataflowFloatVolume::FDataflowFloatVolume(TSharedPtr<UE::DataflowVolume::Private::FDataflowFloatVolumeImpl> InFloatVolume)
+	: FDataflowVolume(InFloatVolume, FDataflowFloatVolume::TypeName)
+{
+	static_assert(sizeof(FDataflowVolume) == sizeof(*this));
+}
+
+/* --------------------------------------------------------------------------------------------------------------- */
+
+TSharedPtr<const UE::DataflowVolume::Private::FDataflowFloatVolumeImpl> FDataflowFloatVolume::GetFloatVolumeImpl() const
+{
+	return StaticCastSharedPtr<const UE::DataflowVolume::Private::FDataflowFloatVolumeImpl>(Volume);
+}
+
+/* --------------------------------------------------------------------------------------------------------------- */
+
+TSharedPtr<UE::DataflowVolume::Private::FDataflowFloatVolumeImpl> FDataflowFloatVolume::GetFloatVolumeImpl()
+{
+	return StaticCastSharedPtr<UE::DataflowVolume::Private::FDataflowFloatVolumeImpl>(Volume);
+}
+
+/* --------------------------------------------------------------------------------------------------------------- */
+
+FVector FDataflowFloatVolume::GetVoxelSize() const
+{
+	using namespace UE::DataflowVolume::Private;
+
+	if (TSharedPtr<const FDataflowFloatVolumeImpl> FloatVolumeImpl = GetFloatVolumeImpl())
+	{
+		return FloatVolumeImpl->GetVoxelSize();
+	}
+
+	return {};
+}
+
+/* --------------------------------------------------------------------------------------------------------------- */
+
+FBox FDataflowFloatVolume::GetVolumeBoundingBox() const
+{
+	using namespace UE::DataflowVolume::Private;
+
+	if (TSharedPtr<const FDataflowFloatVolumeImpl> FloatVolumeImpl = GetFloatVolumeImpl())
+	{
+		return FloatVolumeImpl->GetActiveVoxelBoundingBox();
+	}
+
+	return {};
+}
+
+/* --------------------------------------------------------------------------------------------------------------- */
+
+FDataflowFloatVolume FDataflowFloatVolume::CreateSphereSDF(
+	const float InVoxelSize, 
+	const float InRadius, 
+	const FVector& InCenter)
+{
+	using namespace UE::DataflowVolume::Private;
+
+	TSharedPtr<FDataflowFloatVolumeImpl> VolumeImpl = FDataflowFloatVolumeImpl::CreateSphereSDF(InVoxelSize, InRadius, InCenter);
+	return FDataflowFloatVolume(VolumeImpl);
+}
+
+/* --------------------------------------------------------------------------------------------------------------- */
+
+FDataflowFloatVolume FDataflowFloatVolume::CreateCubeSDF(
+	const float InScale,
+	const FVector& InCenter,
+	const float InVoxelSize)
+{
+	using namespace UE::DataflowVolume::Private;
+
+	TSharedPtr<FDataflowFloatVolumeImpl> VolumeImpl = FDataflowFloatVolumeImpl::CreateCubeSDF(InScale, InCenter, InVoxelSize);
+	return FDataflowFloatVolume(VolumeImpl);
+}
+
+/* --------------------------------------------------------------------------------------------------------------- */
+
+FDataflowFloatVolume FDataflowFloatVolume::CreatePlatonicSolidSDF(const int32 InFaceCount,
+	const float InScale,
+	const FVector& InCenter,
+	const float InVoxelSize)
+{
+	using namespace UE::DataflowVolume::Private;
+
+	TSharedPtr<FDataflowFloatVolumeImpl> VolumeImpl = FDataflowFloatVolumeImpl::CreatePlatonicSolidSDF(InFaceCount, InScale, InCenter, InVoxelSize);
+	return FDataflowFloatVolume(VolumeImpl);
+}
+
+/* --------------------------------------------------------------------------------------------------------------- */
+
+FDataflowFloatVolume FDataflowFloatVolume::CreateVolumeFromMeshDescription(
+	const FMeshDescription& InMeshDescription,
+	const FTransform& InTransform,
+	const float InVoxelSize,
+	const FDataflowVolume& InReferenceVolume,
+	const EDataflowVolumeOutputType InOutputType,
+	const FString& InGridName,
+	const bool InUseWorldSpaceUnits,
+	const float InExteriorBand,
+	const float InInteriorBand,
+	const int32 InExteriorBandVoxels,
+	const int32 InInteriorBandVoxels,
+	const bool bFillInterior,
+	const bool bPreserveHoles,
+	const float InIsoValue,
+	FDataflowIntVolume& OutFaceIndexVolume)
+{
+	using namespace UE::DataflowVolume::Private;
+
+	TSharedPtr<FDataflowIntVolumeImpl> FaceIndexVolumeImpl;
+
+	TSharedPtr<FDataflowFloatVolumeImpl> VolumeImpl = FDataflowFloatVolumeImpl::CreateVolumeFromMeshDescription(
+		InMeshDescription,
+		InTransform,
+		InVoxelSize,
+		InReferenceVolume.GetVolume(),
+		InOutputType,
+		InGridName,
+		InUseWorldSpaceUnits,
+		InExteriorBand,
+		InInteriorBand,
+		InExteriorBandVoxels,
+		InInteriorBandVoxels,
+		bFillInterior,
+		bPreserveHoles,
+		InIsoValue,
+		FaceIndexVolumeImpl);
+	
+	OutFaceIndexVolume.SetIntVolume(FaceIndexVolumeImpl);
+
+	return FDataflowFloatVolume(VolumeImpl);
+}
+
+/* --------------------------------------------------------------------------------------------------------------- */
+
+void FDataflowFloatVolume::VolumeToSpheres(
+	const int32 InMinSphereCount,
+	const int32 InMaxSphereCount,
+	const bool InOverlapping,
+	const float InMinRadius,
+	const float InMaxRadius,
+	const float InIsoValue,
+	const int32 InInstanceCount,
+	TArray<FVector>& SphereCenters,
+	TArray<float>& SphereRadii) const
+{
+	using namespace UE::DataflowVolume::Private;
+
+	if (TSharedPtr<const FDataflowFloatVolumeImpl> FloatVolumeImpl = GetFloatVolumeImpl())
+	{
+		FloatVolumeImpl->VolumeToSpheres(
+			InMinSphereCount,
+			InMaxSphereCount,
+			InOverlapping,
+			InMinRadius,
+			InMaxRadius,
+			InIsoValue,
+			InInstanceCount,
+			SphereCenters,
+			SphereRadii);
+	}
+}
+
+/* --------------------------------------------------------------------------------------------------------------- */
+
+FDataflowFloatVolume FDataflowFloatVolume::CovertSDFToFogVolume(
+	const bool InPruneTolerance,
+	const float InTolerance,
+	const bool InFloodFillOutput,
+	const bool InActivateInterior) const
+{
+	using namespace UE::DataflowVolume::Private;
+	TSharedPtr<FDataflowFloatVolumeImpl> ResultSDF;
+
+	if (TSharedPtr<const FDataflowFloatVolumeImpl> FloatVolumeImpl = GetFloatVolumeImpl())
+	{
+		ResultSDF = FloatVolumeImpl->CovertSDFToFogVolume(
+			InPruneTolerance,
+			InTolerance,
+			InFloodFillOutput,
+			InActivateInterior);
+	}
+
+	return FDataflowFloatVolume(ResultSDF);
+}
+
+/* --------------------------------------------------------------------------------------------------------------- */
+
+FDataflowFloatVolume FDataflowFloatVolume::CovertFogVolumeToSDF(
+	const float InFogIsoValue,
+	const bool InPruneTolerance,
+	const float InTolerance,
+	const bool InFloodFillOutput,
+	const bool InActivateInterior) const
+{
+	using namespace UE::DataflowVolume::Private;
+	TSharedPtr<FDataflowFloatVolumeImpl> ResultSDF;
+
+	if (TSharedPtr<const FDataflowFloatVolumeImpl> FloatVolumeImpl = GetFloatVolumeImpl())
+	{
+		ResultSDF = FloatVolumeImpl->CovertFogVolumeToSDF(
+			InFogIsoValue,
+			InPruneTolerance,
+			InTolerance,
+			InFloodFillOutput,
+			InActivateInterior);
+	}
+
+	return FDataflowFloatVolume(ResultSDF);
+}
+
+/* --------------------------------------------------------------------------------------------------------------- */
+
+void FDataflowFloatVolume::ConvertVolumeToMeshDescription(
+	const float InIsoValue,
+	const float InAdaptivity,
+	FMeshDescription& OutMeshDescription) const
+{
+	using namespace UE::DataflowVolume::Private;
+	TSharedPtr<FDataflowFloatVolumeImpl> ResultSDF;
+
+	if (TSharedPtr<const FDataflowFloatVolumeImpl> FloatVolumeImpl = GetFloatVolumeImpl())
+	{
+		FloatVolumeImpl->ConvertVolumeToMeshDescription(InIsoValue, InAdaptivity, OutMeshDescription);
+	}
+}
+
+/* --------------------------------------------------------------------------------------------------------------- */
+
+void FDataflowFloatVolume::UniformVolumeScatter(
+	const int32 InMinNumberOfPoints,
+	const int32 InMaxNumberOfPoints,
+	const int32 InRandomSeed,
+	const float InIsoValue,
+	const float InSpread,
+	TArray<FVector>& OutPoints) const
+{
+	using namespace UE::DataflowVolume::Private;
+
+	if (TSharedPtr<const FDataflowFloatVolumeImpl> FloatVolumeImpl = GetFloatVolumeImpl())
+	{
+		FloatVolumeImpl->UniformVolumeScatter(
+			InMinNumberOfPoints,
+			InMaxNumberOfPoints,
+			InRandomSeed,
+			InIsoValue,
+			InSpread,
+			OutPoints);
+	}
+}
+
+/* --------------------------------------------------------------------------------------------------------------- */
+
+void FDataflowFloatVolume::DenseUniformVolumeScatter(
+	const float InMinNumberOfPointsPerVoxel,
+	const float InMaxNumberOfPointsPerVoxel,
+	const int32 InRandomSeed,
+	const float InIsoValue,
+	const float InSpread,
+	TArray<FVector>& OutPoints) const
+{
+	using namespace UE::DataflowVolume::Private;
+
+	if (TSharedPtr<const FDataflowFloatVolumeImpl> FloatVolumeImpl = GetFloatVolumeImpl())
+	{
+		FloatVolumeImpl->DenseUniformVolumeScatter(
+			InMinNumberOfPointsPerVoxel,
+			InMaxNumberOfPointsPerVoxel,
+			InRandomSeed,
+			InIsoValue,
+			InSpread,
+			OutPoints);
+	}
+}
+
+/* --------------------------------------------------------------------------------------------------------------- */
+
+void FDataflowFloatVolume::NonUniformVolumeScatter(
+	const float InMinNumberOfPointsPerVoxel,
+	const float InMaxNumberOfPointsPerVoxel,
+	const int32 InRandomSeed,
+	const float InIsoValue,
+	const float InSpread,
+	TArray<FVector>& OutPoints) const
+{
+	using namespace UE::DataflowVolume::Private;
+
+	if (TSharedPtr<const FDataflowFloatVolumeImpl> FloatVolumeImpl = GetFloatVolumeImpl())
+	{
+		FloatVolumeImpl->NonUniformVolumeScatter(
+			InMinNumberOfPointsPerVoxel,
+			InMaxNumberOfPointsPerVoxel,
+			InRandomSeed,
+			InIsoValue,
+			InSpread,
+			OutPoints);
+	}
+}
+
+/* --------------------------------------------------------------------------------------------------------------- */
+
+void FDataflowFloatVolume::VolumeSample(
+	const TArray<FVector>& InPoints,
+	TArray<float>& OutValues) const
+{
+	using namespace UE::DataflowVolume::Private;
+
+	if (TSharedPtr<const FDataflowFloatVolumeImpl> FloatVolumeImpl = GetFloatVolumeImpl())
+	{
+		FloatVolumeImpl->VolumeSample(InPoints, OutValues);
+	}
+}
+
+/* --------------------------------------------------------------------------------------------------------------- */
+
+FDataflowFloatVectorVolume FDataflowFloatVolume::ComputeGradient(const EDataflowVolumeAnalysisOutputName OutputName, const FString& CustomName, const FDataflowBoolVolume* MaskVolume) const
+{
+	using namespace UE::DataflowVolume::Private;
+
+	if (TSharedPtr<const FDataflowFloatVolumeImpl> FloatVolumeImpl = GetFloatVolumeImpl())
+	{
+		if (TSharedPtr<const FDataflowBoolVolumeImpl> MaskVolumeImpl = MaskVolume ? MaskVolume->GetBoolVolumeImpl() : MakeShared<FDataflowBoolVolumeImpl>())
+		{
+			TSharedPtr<FDataflowFloatVectorVolumeImpl> GradientVolume = FloatVolumeImpl->ComputeGradient(OutputName, CustomName, MaskVolumeImpl);
+			return FDataflowFloatVectorVolume(GradientVolume);
+		}
+	}
+
+	return FDataflowFloatVectorVolume();
+}
+
+/* --------------------------------------------------------------------------------------------------------------- */
+
+FDataflowFloatVolume FDataflowFloatVolume::ComputeCurvature(const EDataflowVolumeAnalysisOutputName OutputName, const FString& CustomName, const FDataflowBoolVolume* MaskVolume) const
+{
+	using namespace UE::DataflowVolume::Private;
+
+	if (TSharedPtr<const FDataflowFloatVolumeImpl> FloatVolumeImpl = GetFloatVolumeImpl())
+	{
+		if (TSharedPtr<const FDataflowBoolVolumeImpl> MaskVolumeImpl = MaskVolume ? MaskVolume->GetBoolVolumeImpl() : MakeShared<FDataflowBoolVolumeImpl>())
+		{
+			TSharedPtr<FDataflowFloatVolumeImpl> CurvatureVolume = FloatVolumeImpl->ComputeCurvature(OutputName, CustomName, MaskVolumeImpl);
+			return FDataflowFloatVolume(CurvatureVolume);
+		}
+	}
+
+	return FDataflowFloatVolume();
+}
+
+/* --------------------------------------------------------------------------------------------------------------- */
+
+FDataflowFloatVolume FDataflowFloatVolume::ComputeLaplacian(const EDataflowVolumeAnalysisOutputName OutputName, const FString& CustomName, const FDataflowBoolVolume* MaskVolume) const
+{
+	using namespace UE::DataflowVolume::Private;
+
+	if (TSharedPtr<const FDataflowFloatVolumeImpl> FloatVolumeImpl = GetFloatVolumeImpl())
+	{
+		if (TSharedPtr<const FDataflowBoolVolumeImpl> MaskVolumeImpl = MaskVolume ? MaskVolume->GetBoolVolumeImpl() : MakeShared<FDataflowBoolVolumeImpl>())
+		{
+			TSharedPtr<FDataflowFloatVolumeImpl> LaplacianVolume = FloatVolumeImpl->ComputeLaplacian(OutputName, CustomName, MaskVolumeImpl);
+			return FDataflowFloatVolume(LaplacianVolume);
+		}
+	}
+
+	return FDataflowFloatVolume();
+}
+
+/* --------------------------------------------------------------------------------------------------------------- */
+
+FDataflowFloatVectorVolume FDataflowFloatVolume::ComputeClosestPoint(const EDataflowVolumeAnalysisOutputName OutputName, const FString& CustomName, const FDataflowBoolVolume* MaskVolume) const
+{
+	using namespace UE::DataflowVolume::Private;
+
+	if (TSharedPtr<const FDataflowFloatVolumeImpl> FloatVolumeImpl = GetFloatVolumeImpl())
+	{
+		if (TSharedPtr<const FDataflowBoolVolumeImpl> MaskVolumeImpl = MaskVolume ? MaskVolume->GetBoolVolumeImpl() : MakeShared<FDataflowBoolVolumeImpl>())
+		{
+			TSharedPtr<FDataflowFloatVectorVolumeImpl> ClosestPointVolume = FloatVolumeImpl->ComputeClosestPoint(OutputName, CustomName, MaskVolumeImpl);
+			return FDataflowFloatVectorVolume(ClosestPointVolume);
+		}
+	}
+
+	return FDataflowFloatVectorVolume();
+}
+
+/* --------------------------------------------------------------------------------------------------------------- */
+
+FDataflowFloatVolume FDataflowFloatVolume::VolumeCombine(
+	const FDataflowFloatVolume& InVolumeB,
+	const EDataflowVolumeSDFCombineOperation InOperation,
+	const float InMultiplierA,
+	const float InMultiplierB,
+	const EDataflowVolumeSDFCombineResample InResample,
+	const EDataflowVolumeSDFCombineInterpolation InInterpolation,
+	const bool InPruneDegenerateTiles,
+	const float InPruneTol) const
+{
+	using namespace UE::DataflowVolume::Private;
+
+	if (TSharedPtr<const FDataflowFloatVolumeImpl> FloatVolumeAImpl = GetFloatVolumeImpl())
+	{
+		if (TSharedPtr<const FDataflowFloatVolumeImpl> FloatVolumeBImpl = InVolumeB.GetFloatVolumeImpl())
+		{
+			TSharedPtr<FDataflowFloatVolumeImpl> ResultVolumeImpl = FloatVolumeAImpl->VolumeCombine(
+				FloatVolumeBImpl,
+				InOperation,
+				InMultiplierA,
+				InMultiplierB,
+				InResample,
+				InInterpolation,
+				InPruneDegenerateTiles,
+				InPruneTol);
+
+			return FDataflowFloatVolume(ResultVolumeImpl);
+		}
+	}
+
+	return FDataflowFloatVolume();
+}
diff --git a/Engine/Plugins/Experimental/DataflowVolumetricPlugin/Source/DataflowVolumeCore/Private/Dataflow/DataflowFloatVolumeImpl.cpp b/Engine/Plugins/Experimental/DataflowVolumetricPlugin/Source/DataflowVolumeCore/Private/Dataflow/DataflowFloatVolumeImpl.cpp
new file mode 100644
--- /dev/null
+++ b/Engine/Plugins/Experimental/DataflowVolumetricPlugin/Source/DataflowVolumeCore/Private/Dataflow/DataflowFloatVolumeImpl.cpp
@@ -0,0 +1,1196 @@
+// Copyright Epic Games, Inc. All Rights Reserved.
+
+#include "Dataflow/DataflowFloatVolumeImpl.h"
+#include "Dataflow/DataflowVolume.h"
+#include "Dataflow/DataflowIntVolume.h"
+#include "Dataflow/DataflowIntVolumeImpl.h"
+#include "Dataflow/DataflowBoolVolumeImpl.h"
+#include "Dataflow/DataflowVolumeUtils.h"
+
+namespace UE::DataflowVolume::Private
+{
+	/* --------------------------------------------------------------------------------------------------------------- */
+	/* FDataflowFloatVolumeImpl */
+	/* --------------------------------------------------------------------------------------------------------------- */
+
+	openvdb::FloatGrid::Ptr FDataflowFloatVolumeImpl::GetGrid() const
+	{
+		return FloatGrid;
+	}
+
+	/* --------------------------------------------------------------------------------------------------------------- */
+
+	void FDataflowFloatVolumeImpl::SetGrid(const openvdb::FloatGrid::Ptr& InFloatGrid, const openvdb::GridClass InGridClass, const FString& InGridName)
+	{
+		FDataflowVolumeImpl::SetGrid(InFloatGrid, InGridName);
+
+		FloatGrid = InFloatGrid;
+		FloatGrid->setGridClass(InGridClass);
+	}
+
+	/* --------------------------------------------------------------------------------------------------------------- */
+
+	FString FDataflowFloatVolumeImpl::GetGridType() const
+	{
+		if (FloatGrid)
+		{
+			std::string GridType = FloatGrid->gridType();
+			return FString(GridType.c_str());
+		}
+
+		return {};
+	}
+
+	/* --------------------------------------------------------------------------------------------------------------- */
+
+	FString FDataflowFloatVolumeImpl::GetGridClass() const
+	{
+		if (FloatGrid)
+		{
+			FString GridClass;
+			switch (FloatGrid->getGridClass()) {
+			case openvdb::GRID_UNKNOWN:
+				GridClass = TEXT("Unknown Grid Class");
+				break;
+			case openvdb::GRID_LEVEL_SET:
+				GridClass = TEXT("Level Set");
+				break;
+			case openvdb::GRID_FOG_VOLUME:
+				GridClass = TEXT("Fog Volume");
+				break;
+			case openvdb::GRID_STAGGERED:
+				GridClass = TEXT("Staggered Vector Field");
+				break;
+
+			default:
+				break;
+			}
+
+			return GridClass;
+		}
+
+		return {};
+	}
+
+	/* --------------------------------------------------------------------------------------------------------------- */
+
+	FString FDataflowFloatVolumeImpl::VolumeInfo() const
+	{
+		FString BaseString = FDataflowVolumeImpl::VolumeInfo();
+
+		FString OutputStr;
+		OutputStr += BaseString;
+
+		FString GridClass = this->GetGridClass();
+		FString GridType = this->GetGridType();
+
+		OutputStr.Appendf(TEXT("Class: %s\nType: %s\n"),
+			*GridClass,
+			*GridType);
+
+		return OutputStr;
+	}
+
+	/* --------------------------------------------------------------------------------------------------------------- */
+
+	TSharedPtr<FDataflowFloatVolumeImpl> FDataflowFloatVolumeImpl::CreateSphereSDF(
+		const float InVoxelSize,
+		const float InRadius,
+		const FVector& InCenter)
+	{
+		TSharedPtr<FDataflowFloatVolumeImpl> SDFImpl = MakeShared<FDataflowFloatVolumeImpl>();
+
+		// Generate a level set grid
+		openvdb::FloatGrid::Ptr SphereGrid = openvdb::tools::createLevelSetSphere<openvdb::FloatGrid>(
+			/*radius=*/InRadius,
+			/*center=*/openvdb::Vec3f((float)InCenter.X, (float)InCenter.Y, (float)InCenter.Z),
+			/*voxel size=*/InVoxelSize);
+
+		SphereGrid->insertMeta("radius", openvdb::FloatMetadata(InRadius));
+		SphereGrid->insertMeta("center", openvdb::Vec3SMetadata(openvdb::Vec3s((float)InCenter.X, (float)InCenter.Y, (float)InCenter.Z)));
+
+		SDFImpl->SetGrid(SphereGrid, openvdb::GRID_LEVEL_SET);
+
+		return SDFImpl;
+	}
+
+	/* --------------------------------------------------------------------------------------------------------------- */
+
+	TSharedPtr<FDataflowFloatVolumeImpl> FDataflowFloatVolumeImpl::CreateCubeSDF(const float InScale, const FVector& InCenter, const float InVoxelSize)
+	{
+		TSharedPtr<FDataflowFloatVolumeImpl> SDFImpl = MakeShared<FDataflowFloatVolumeImpl>();
+
+		const openvdb::Vec3f CenterVec = openvdb::Vec3f(InCenter.X, InCenter.Y, InCenter.Z);
+
+		// Generate a level set grid
+		openvdb::FloatGrid::Ptr BoxGrid = openvdb::tools::createLevelSetCube<openvdb::FloatGrid>(
+			/*scale of the cube in world units*/InScale,
+			/*center of the cube in world units*/CenterVec,
+			/*voxel size in world units*/InVoxelSize);
+
+		SDFImpl->SetGrid(BoxGrid, openvdb::GRID_LEVEL_SET);
+
+		return SDFImpl;
+	}
+
+	/* --------------------------------------------------------------------------------------------------------------- */
+
+	TSharedPtr<FDataflowFloatVolumeImpl> FDataflowFloatVolumeImpl::CreatePlatonicSolidSDF(const int32 InFaceCount,
+		const float InScale,
+		const FVector& InCenter,
+		const float InVoxelSize)
+	{
+		TSharedPtr<FDataflowFloatVolumeImpl> SDFImpl = MakeShared<FDataflowFloatVolumeImpl>();
+
+		// Generate a level set grid
+		const openvdb::Vec3f CenterVec = openvdb::Vec3f(InCenter.X, InCenter.Y, InCenter.Z);
+		openvdb::FloatGrid::Ptr PlatonicGrid = openvdb::tools::createLevelSetPlatonic<openvdb::FloatGrid>(
+			/*number of faces of the platonic solid, i.e. 4, 6, 8, 12 or 20*/InFaceCount,
+			/*scale of the platonic solid in world units*/InScale,
+			/*center of the platonic solid in world units*/CenterVec,
+			/*voxel size in world units*/InVoxelSize);
+
+		openvdb::math::Transform GridTransform = PlatonicGrid->transform();
+
+		SDFImpl->SetGrid(PlatonicGrid, openvdb::GRID_LEVEL_SET);
+
+		return SDFImpl;
+	}
+
+	/* --------------------------------------------------------------------------------------------------------------- */
+
+	TSharedPtr<FDataflowFloatVolumeImpl> FDataflowFloatVolumeImpl::CreateVolumeFromMeshDescription(
+		const FMeshDescription& InMeshDescription,
+		const FTransform& InTransform,
+		const float InVoxelSize,
+		const TSharedPtr<const FDataflowVolumeImpl>& InReferenceVolumeImpl,
+		const EDataflowVolumeOutputType InOutputType,
+		const FString& InGridName,
+		const bool InUseWorldSpaceUnits,
+		const float InExteriorBand,
+		const float InInteriorBand,
+		const int32 InExteriorBandVoxels,
+		const int32 InInteriorBandVoxels,
+		const bool bFillInterior,
+		const bool bPreserveHoles,
+		const float InIsoValue,
+		TSharedPtr<FDataflowIntVolumeImpl>& OutFaceIndexVolume)
+	{
+		using namespace UE::DataflowVolume::Private;
+
+		TSharedPtr<FDataflowFloatVolumeImpl> FloatVolumeImpl = MakeShared<FDataflowFloatVolumeImpl>();
+
+		// If there is a reference grid connected, derive the voxel size and transform
+		float VoxelSize = InVoxelSize;
+
+		// @todo(gmelich): Implement using a reference grid
+		if (InReferenceVolumeImpl.IsValid() && InReferenceVolumeImpl->IsGridEmpty())
+		{
+		}
+
+		UE::DataflowVolumeUtils::FPlacedMesh PlacedMesh(&InMeshDescription, InTransform);
+
+		FMatrix LocalToVoxel = FMatrix::Identity;
+		LocalToVoxel.M[0][0] = InVoxelSize;
+		LocalToVoxel.M[1][1] = InVoxelSize;
+		LocalToVoxel.M[2][2] = InVoxelSize;
+
+		FVector AverageTranslation = FVector(0.0);
+
+		auto TransformGenerator = [&LocalToVoxel, &AverageTranslation](const UE::DataflowVolumeUtils::FPlacedMesh& PlacedMesh)->openvdb::Mat4R
+			{
+				FTransform MeshXForm = PlacedMesh.Transform;
+				MeshXForm.AddToTranslation(-AverageTranslation);
+				FMatrix TransformMatrix = MeshXForm.ToMatrixWithScale().Inverse();
+
+				TransformMatrix = LocalToVoxel * TransformMatrix;
+				double* data = &TransformMatrix.M[0][0];
+				openvdb::Mat4R VDBMatDouble(data);
+				// NB: rounding errors in the inverse may have resulted in error in this col.
+				// openvdb explicitly checks this matrix row to insure the transform is affine and will throw 
+				VDBMatDouble.setCol(3, openvdb::Vec4R(0, 0, 0, 1));
+				return VDBMatDouble;
+			};
+
+		openvdb::Mat4R XFormA = TransformGenerator(PlacedMesh);
+		openvdb::math::Transform::Ptr VDBXForm = openvdb::math::Transform::createLinearTransform(XFormA);
+
+		// Create adapters that understand the openVDB semantics
+		UE::DataflowVolumeUtils::FMeshDescriptionAdapter Adapter(InMeshDescription, *VDBXForm);
+
+		// target transform
+		openvdb::math::Transform::Ptr TargetXForm = openvdb::math::Transform::createLinearTransform(InVoxelSize);
+		//	openvdb::math::Transform::Ptr TargetXForm = openvdb::math::Transform::createLinearTransform(XFormA);
+
+		if (VoxelSize < 1e-5)
+		{
+			FloatVolumeImpl->SetGrid(openvdb::FloatGrid::Ptr(), openvdb::GRID_LEVEL_SET);
+
+			return FloatVolumeImpl;
+		}
+
+		// Compute band width
+		float ExteriorBand = 0.f;
+		float InteriorBand = FLT_MAX;
+
+		if (InUseWorldSpaceUnits)
+		{
+			ExteriorBand = InExteriorBand / VoxelSize;
+		}
+		else
+		{
+			ExteriorBand = InExteriorBandVoxels;
+		}
+
+		if (!bFillInterior)
+		{
+			if (InUseWorldSpaceUnits)
+			{
+				InteriorBand = InInteriorBand / VoxelSize;
+			}
+			else
+			{
+				InteriorBand = InInteriorBandVoxels;
+			}
+		}
+
+		// @todo(gmelich): Implement using interrupter
+		// Interrupter
+		UE::DataflowVolumeUtils::FInterrupter NullInterrupter;
+
+		// Grid for storing closest face index
+		openvdb::Int32Grid::Ptr FaceIndexGrid;
+		FaceIndexGrid.reset(new openvdb::Int32Grid(0));
+
+		// Compute output
+		openvdb::FloatGrid::Ptr ResultGrid;
+
+		openvdb::GridClass Class = openvdb::GridClass::GRID_UNKNOWN;
+		if (InOutputType == EDataflowVolumeOutputType::SDF)
+		{
+			int32 Flags = 0;
+			ResultGrid = openvdb::tools::meshToVolume<openvdb::FloatGrid>(NullInterrupter, Adapter, *TargetXForm, ExteriorBand, InteriorBand, Flags, FaceIndexGrid.get());
+			Class = openvdb::GRID_LEVEL_SET;
+		}
+		else if (InOutputType == EDataflowVolumeOutputType::USDF)
+		{
+			int32 Flags = openvdb::tools::UNSIGNED_DISTANCE_FIELD;
+			ResultGrid = openvdb::tools::meshToVolume<openvdb::FloatGrid>(NullInterrupter, Adapter, *TargetXForm, ExteriorBand, InteriorBand, Flags, FaceIndexGrid.get());
+			Class = openvdb::GRID_LEVEL_SET;
+		}
+		else if (InOutputType == EDataflowVolumeOutputType::FogVolume)
+		{
+			int32 Flags = 0;
+			ResultGrid = openvdb::tools::meshToVolume<openvdb::FloatGrid>(NullInterrupter, Adapter, *TargetXForm, ExteriorBand, InteriorBand, Flags, FaceIndexGrid.get());
+			openvdb::tools::sdfToFogVolume(*ResultGrid);
+			Class = openvdb::GRID_FOG_VOLUME;
+		}
+
+		// Set output grids
+		FloatVolumeImpl->SetGrid(ResultGrid, Class, InGridName);
+
+		OutFaceIndexVolume = MakeShared<FDataflowIntVolumeImpl>();
+		if (OutFaceIndexVolume)
+		{
+			OutFaceIndexVolume->SetGrid(FaceIndexGrid, "FaceIndex");
+		}
+
+		return FloatVolumeImpl;
+	}
+
+	/* --------------------------------------------------------------------------------------------------------------- */
+
+	void FDataflowFloatVolumeImpl::VolumeToSpheres(const int32 InMinSphereCount,
+		const int32 InMaxSphereCount,
+		const bool InOverlapping,
+		const float InMinRadius,
+		float InMaxRadius,
+		const float InIsoValue,
+		const int32 InInstanceCount,
+		TArray<FVector>& SphereCenters,
+		TArray<float>& SphereRadii) const
+	{
+		if (FloatGrid)
+		{
+			openvdb::FloatGrid::ConstPtr GridPtr = openvdb::gridConstPtrCast<openvdb::FloatGrid>(FloatGrid);
+
+			std::vector<openvdb::Vec4s> Spheres;
+			const openvdb::Vec2i SphereCount(InMinSphereCount, InMaxSphereCount);
+			if (InMaxRadius <= InMinRadius)
+			{
+				InMaxRadius = InMinRadius + (float)1e-5;
+			}
+
+			openvdb::tools::fillWithSpheres(*GridPtr,
+				Spheres, SphereCount, InOverlapping, InMinRadius, InMaxRadius, InIsoValue, InInstanceCount);
+
+			const int32 NumSpheres = Spheres.size();
+			SphereCenters.AddUninitialized(NumSpheres);
+			SphereRadii.AddUninitialized(NumSpheres);
+
+			for (int32 Idx = 0; Idx < NumSpheres; ++Idx)
+			{
+				SphereCenters[Idx] = FVector(Spheres[Idx].x(), Spheres[Idx].y(), Spheres[Idx].z());
+				SphereRadii[Idx] = Spheres[Idx].w();
+			}
+		}
+	}
+
+	/* --------------------------------------------------------------------------------------------------------------- */
+
+	TSharedPtr<FDataflowFloatVolumeImpl> FDataflowFloatVolumeImpl::CovertSDFToFogVolume(
+		const bool InPruneTolerance,
+		const float InTolerance,
+		const bool InFloodFillOutput,
+		const bool InActivateInteriror) const
+	{
+		TSharedPtr<FDataflowFloatVolumeImpl> VolumeImpl = MakeShared<FDataflowFloatVolumeImpl>();
+
+		openvdb::FloatGrid::Ptr ResultGrid = this->FloatGrid;
+
+		// Compute
+		openvdb::tools::sdfToFogVolume(*ResultGrid);
+		VolumeImpl->SetGrid(ResultGrid, openvdb::GRID_FOG_VOLUME, "FogVolume");
+
+		return VolumeImpl;
+	}
+
+	/* --------------------------------------------------------------------------------------------------------------- */
+
+	TSharedPtr<FDataflowFloatVolumeImpl> FDataflowFloatVolumeImpl::CovertFogVolumeToSDF(
+		const float InFogIsoValue,
+		const bool InPruneTolerance,
+		const float InTolerance,
+		const bool InFloodFillOutput,
+		const bool InActivateInteriror) const
+	{
+		TSharedPtr<FDataflowFloatVolumeImpl> VolumeImpl = MakeShared<FDataflowFloatVolumeImpl>();
+
+		openvdb::FloatGrid::Ptr ResultGrid = this->FloatGrid;
+		//FMeshDescription MeshDescription,
+
+		// @todo(gmelich): Fix this
+
+		// Step 1 - Convert Fog volume to mesh
+		//UE::DataflowVolumeUtils::FMixedPolyMesh TempMixedPolyMesh;
+		//UE::DataflowVolumeUtils::FAOSMesh AOSMeshedVolume;
+
+		//openvdb::tools::volumeToMesh(*this->FloatGrid, TempMixedPolyMesh.Points, TempMixedPolyMesh.Triangles, TempMixedPolyMesh.Quads, InFogIsoValue, InAdaptivity);
+
+		//UE::DataflowVolumeUtils::MixedPolyMeshToAOSMesh(TempMixedPolyMesh, AOSMeshedVolume);
+
+		//ConvertMesh(AOSMeshedVolume, MeshDescription);
+
+		//// Step 2 - Convert mesh to SDF
+		//float VoxelSize = ResultGrid->voxelSize()[0];
+
+		//FMatrix LocalToVoxel = FMatrix::Identity;
+		//LocalToVoxel.M[0][0] = VoxelSize;
+		//LocalToVoxel.M[1][1] = VoxelSize;
+		//LocalToVoxel.M[2][2] = VoxelSize;
+
+		//UE::DataflowVolumeUtils::FPlacedMesh PlacedMesh(&MeshDescription, FTransform::Identity);
+		//FVector AverageTranslation = FVector(0.0);
+
+		//auto TransformGenerator = [&LocalToVoxel, &AverageTranslation](const UE::DataflowVolumeUtils::FPlacedMesh& PlacedMesh)->openvdb::Mat4R
+		//	{
+		//		FTransform MeshXForm = PlacedMesh.Transform;
+		//		MeshXForm.AddToTranslation(-AverageTranslation);
+		//		FMatrix TransformMatrix = MeshXForm.ToMatrixWithScale().Inverse();
+
+		//		TransformMatrix = LocalToVoxel * TransformMatrix;
+		//		double* data = &TransformMatrix.M[0][0];
+		//		openvdb::Mat4R VDBMatDouble(data);
+		//		// NB: rounding errors in the inverse may have resulted in error in this col.
+		//		// openvdb explicitly checks this matrix row to insure the transform is affine and will throw 
+		//		VDBMatDouble.setCol(3, openvdb::Vec4R(0, 0, 0, 1));
+		//		return VDBMatDouble;
+		//	};
+
+		//openvdb::Mat4R XFormA = TransformGenerator(PlacedMesh);
+		//openvdb::math::Transform::Ptr VDBXForm = openvdb::math::Transform::createLinearTransform(XFormA);
+
+		//// Create adapters that understand the openVDB semantics
+		//UE::DataflowVolumeUtils::FMeshDescriptionAdapter Adapter(MeshDescription, *VDBXForm);
+
+		//openvdb::math::Transform::Ptr TargetXForm = openvdb::math::Transform::createLinearTransform(VoxelSize);
+
+		//UE::DataflowVolumeUtils::FInterrupter NullInterrupter;
+
+		//int32 Flags = 0;
+		//ResultGrid = openvdb::tools::meshToVolume<openvdb::FloatGrid>(NullInterrupter, Adapter, *TargetXForm);
+
+		//VolumeImpl->SetGrid(ResultGrid, openvdb::GRID_LEVEL_SET, "LevelSet");
+
+		VolumeImpl->SetGrid(ResultGrid, openvdb::GRID_FOG_VOLUME, "FogVolume");
+		return VolumeImpl;
+	}
+
+	/* --------------------------------------------------------------------------------------------------------------- */
+
+	void FDataflowFloatVolumeImpl::CovertSDFToCollection(
+		const float InIsoValue,
+		const float InAdaptivity,
+		FManagedArrayCollection& OutCollection) const
+	{
+		openvdb::FloatGrid::Ptr ResultGrid;
+
+		// Compute
+		openvdb::tools::VolumeToMesh Mesher(InIsoValue, InAdaptivity);
+		Mesher(*this->FloatGrid);
+
+		const int32 NumPoints = Mesher.pointListSize();
+		const int32 NumPolys = Mesher.polygonPoolListSize();
+
+		openvdb::tools::PointList& Points = Mesher.pointList();
+		openvdb::tools::PolygonPoolList& Polys = Mesher.polygonPoolList();
+	}
+
+	/* --------------------------------------------------------------------------------------------------------------- */
+
+	void FDataflowFloatVolumeImpl::ConvertVolumeToMeshDescription(
+		const float InIsoValue,
+		const float InAdaptivity,
+		FMeshDescription& OutMeshDescription) const
+	{
+		openvdb::FloatGrid::Ptr ResultGrid;
+
+		UE::DataflowVolumeUtils::FMixedPolyMesh TempMixedPolyMesh;
+		UE::DataflowVolumeUtils::FAOSMesh AOSMeshedVolume;
+
+		openvdb::tools::volumeToMesh(*this->FloatGrid, TempMixedPolyMesh.Points, TempMixedPolyMesh.Triangles, TempMixedPolyMesh.Quads, InIsoValue, InAdaptivity);
+
+		UE::DataflowVolumeUtils::MixedPolyMeshToAOSMesh(TempMixedPolyMesh, AOSMeshedVolume);
+
+		ConvertMesh(AOSMeshedVolume, OutMeshDescription);
+	}
+
+	/* --------------------------------------------------------------------------------------------------------------- */
+
+	void FDataflowFloatVolumeImpl::GetActiveVoxels(TArray<FBox>& OutActiveVoxels) const
+	{
+		const openvdb::Index64 ActiveVoxelCount = FloatGrid->activeVoxelCount();
+		OutActiveVoxels.AddUninitialized(ActiveVoxelCount);
+
+		openvdb::math::Transform T = this->FloatGrid->transform();
+
+		int32 Idx = 0;
+		for (openvdb::FloatGrid::ValueOnCIter iter = this->FloatGrid->cbeginValueOn(); iter.test(); ++iter)
+		{
+			if (iter.isVoxelValue())
+			{
+				openvdb::CoordBBox BBox;
+				iter.getBoundingBox(BBox);
+
+				const openvdb::Vec3d BBoxMin(BBox.min().x() - 0.5, BBox.min().y() - 0.5, BBox.min().z() - 0.5);
+				const openvdb::Vec3d BBoxMax(BBox.max().x() + 0.5, BBox.max().y() + 0.5, BBox.max().z() + 0.5);
+
+				const openvdb::Vec3d BBoxMinWS = FloatGrid->transform().indexToWorld(BBoxMin);
+				const openvdb::Vec3d BBoxMaxWS = FloatGrid->transform().indexToWorld(BBoxMax);
+
+				FVector Min(BBoxMinWS.x(), BBoxMinWS.y(), BBoxMinWS.z());
+				FVector Max(BBoxMaxWS.x(), BBoxMaxWS.y(), BBoxMaxWS.z());
+
+				OutActiveVoxels[Idx++] = FBox(Min, Max);
+			}
+		}
+	}
+
+	/* --------------------------------------------------------------------------------------------------------------- */
+
+	typedef openvdb::FloatGrid GridT;
+	using TreeT = typename GridT::TreeType;
+	using BoolTreeT = typename TreeT::template ValueConverter<bool>::Type;
+
+	static openvdb::Grid<BoolTreeT>::Ptr GetVolumeForScatter(
+		const openvdb::FloatGrid::Ptr& InFloatGrid,
+		const float InIsoValue)
+	{
+		using namespace openvdb;
+
+		openvdb::FloatGrid::Ptr gridPtr = InFloatGrid->copy(); // shallow copy
+
+		float Isovalue = 0.0f;
+
+		if (gridPtr->getGridClass() == GRID_LEVEL_SET)
+		{
+			// Clamp the isovalue to the level set's background value minus epsilon.
+			// (In a valid narrow-band level set, all voxels, including background voxels,
+			// have values less than or equal to the background value, so an isovalue
+			// greater than or equal to the background value would produce a mask with
+			// effectively infinite extent.)
+			Isovalue = std::min(InIsoValue,
+				static_cast<float>(gridPtr->background() - math::Tolerance<float>::value()));
+		}
+		else if (gridPtr->getGridClass() == GRID_FOG_VOLUME)
+		{
+			// Clamp the isovalue of a fog volume between epsilon and one,
+			// again to avoid a mask with infinite extent.  (Recall that
+			// fog volume voxel values vary from zero outside to one inside.)
+			Isovalue = math::Clamp(InIsoValue, math::Tolerance<float>::value(), 1.f);
+		}
+
+		int32 numVoxels = gridPtr->activeVoxelCount();
+		if (numVoxels < 10000)
+		{
+			const double scale = 1.0 / math::Cbrt(2.0 * 10000.0 / double(numVoxels));
+			auto scaledXform = gridPtr->transform().copy();
+			scaledXform->preScale(scale);
+
+			auto newGridPtr = openvdb::tools::levelSetRebuild(*gridPtr, Isovalue,
+				LEVEL_SET_HALF_WIDTH, LEVEL_SET_HALF_WIDTH, scaledXform.get());
+
+			const int32 newNumVoxels = newGridPtr->activeVoxelCount();
+			if (newNumVoxels > numVoxels)
+			{
+				gridPtr = newGridPtr;
+				numVoxels = newNumVoxels;
+			}
+		}
+
+		//			typedef openvdb::FloatGrid GridT;
+
+		//			using TreeT = typename GridT::TreeType;
+		//			using BoolTreeT = typename TreeT::template ValueConverter<bool>::Type;
+
+		const TreeT& tree = gridPtr->tree();
+		math::Transform transform = gridPtr->transform();
+
+		// Compute a mask of the voxels enclosed by the isosurface.
+		typename Grid<BoolTreeT>::Ptr InteriorMaskPtr;
+
+		if (gridPtr->getGridClass() == GRID_LEVEL_SET) {
+			InteriorMaskPtr = tools::sdfInteriorMask(*gridPtr, Isovalue);
+		}
+		else
+		{
+			// For non-level-set grids, the interior mask comprises the active voxels.
+			InteriorMaskPtr = typename Grid<BoolTreeT>::Ptr(Grid<BoolTreeT>::create(false));
+			InteriorMaskPtr->setTransform(transform.copy());
+			InteriorMaskPtr->tree().topologyUnion(tree);
+		}
+
+		{
+			auto& maskTree = InteriorMaskPtr->tree();
+			auto copyOfTree = StaticPtrCast<BoolTreeT>(maskTree.copy());
+			tools::erodeActiveValues(maskTree, /*iterations=*/1, tools::NN_FACE, tools::IGNORE_TILES);
+			tools::pruneInactive(maskTree);
+			if (maskTree.empty()) { InteriorMaskPtr->setTree(copyOfTree); }
+		}
+
+		return InteriorMaskPtr;
+	}
+
+	/* --------------------------------------------------------------------------------------------------------------- */
+
+	void FDataflowFloatVolumeImpl::UniformVolumeScatter(
+		const int32 InMinNumberOfPoints,
+		const int32 InMaxNumberOfPoints,
+		const int32 InRandomSeed,
+		const float InIsoValue,
+		const float InSpread,
+		TArray<FVector>& OutPoints) const
+	{
+		using namespace openvdb;
+
+		OutPoints.Empty();
+
+		if (this->FloatGrid->empty() ||
+			!(this->FloatGrid->getGridClass() == GRID_LEVEL_SET || this->FloatGrid->getGridClass() == GRID_FOG_VOLUME))
+		{
+			return;
+		}
+
+		openvdb::Grid<BoolTreeT>::Ptr InteriorMaskPtr = GetVolumeForScatter(this->FloatGrid, InIsoValue);
+
+		FRandomStream RandStream(InRandomSeed);
+		int32 NumPoints = InMinNumberOfPoints;
+		if (InMaxNumberOfPoints > InMinNumberOfPoints)
+		{
+			NumPoints = RandStream.RandRange(InMinNumberOfPoints, InMaxNumberOfPoints);
+		}
+
+		using RandGen = std::mersenne_twister_engine<uint32_t, 32, 351, 175, 19,
+			0xccab8ee7, 11, 0xffffffff, 7, 0x31b6ab00, 15, 0xffe50000, 17, 1812433253>; // mt11213b
+		RandGen MTRand(/*seed=*/InRandomSeed);
+
+		std::vector<Vec3R> InstancePoints;
+		InstancePoints.reserve(NumPoints);
+		tools::v2s_internal::PointAccessor PointAcc(InstancePoints);
+
+		tools::UniformPointScatter<tools::v2s_internal::PointAccessor, RandGen> UniformScatter(PointAcc, Index64(NumPoints), MTRand, double(InSpread));
+		UniformScatter(*InteriorMaskPtr);
+
+		for (std::vector<Vec3R>::iterator It = InstancePoints.begin(); It != InstancePoints.end(); It++)
+		{
+			const FVector Position = FVector(It->x(), It->y(), It->z());
+			OutPoints.Add(Position);
+		}
+		return;
+	}
+
+	/* --------------------------------------------------------------------------------------------------------------- */
+
+	void FDataflowFloatVolumeImpl::DenseUniformVolumeScatter(
+		const float InMinNumberOfPointsPerVoxel,
+		const float InMaxNumberOfPointsPerVoxel,
+		const int32 InRandomSeed,
+		const float InIsoValue,
+		const float InSpread,
+		TArray<FVector>& OutPoints) const
+	{
+		using namespace openvdb;
+
+		OutPoints.Empty();
+
+		if (this->FloatGrid->empty() ||
+			!(this->FloatGrid->getGridClass() == GRID_LEVEL_SET || this->FloatGrid->getGridClass() == GRID_FOG_VOLUME))
+		{
+			return;
+		}
+
+		openvdb::Grid<BoolTreeT>::Ptr InteriorMaskPtr = GetVolumeForScatter(this->FloatGrid, InIsoValue);
+
+		FRandomStream RandStream(InRandomSeed);
+		float NumPointsPerVoxel = InMinNumberOfPointsPerVoxel;
+		if (InMaxNumberOfPointsPerVoxel > InMinNumberOfPointsPerVoxel)
+		{
+			NumPointsPerVoxel = RandStream.FRandRange(InMinNumberOfPointsPerVoxel, InMaxNumberOfPointsPerVoxel);
+		}
+
+		using RandGen = std::mersenne_twister_engine<uint32_t, 32, 351, 175, 19,
+			0xccab8ee7, 11, 0xffffffff, 7, 0x31b6ab00, 15, 0xffe50000, 17, 1812433253>; // mt11213b
+		RandGen MTRand(/*seed=*/InRandomSeed);
+
+		std::vector<Vec3R> InstancePoints;
+		//			InstancePoints.reserve(NumPoints);
+		tools::v2s_internal::PointAccessor PointAcc(InstancePoints);
+
+		tools::DenseUniformPointScatter<tools::v2s_internal::PointAccessor, RandGen> DenseUniformScatter(PointAcc, NumPointsPerVoxel, MTRand, double(InSpread));
+		DenseUniformScatter(*InteriorMaskPtr);
+
+		for (std::vector<Vec3R>::iterator It = InstancePoints.begin(); It != InstancePoints.end(); It++)
+		{
+			const FVector Position = FVector(It->x(), It->y(), It->z());
+			OutPoints.Add(Position);
+		}
+		return;
+	}
+
+	/* --------------------------------------------------------------------------------------------------------------- */
+
+	void FDataflowFloatVolumeImpl::NonUniformVolumeScatter(
+		const float InMinNumberOfPointsPerVoxel,
+		const float InMaxNumberOfPointsPerVoxel,
+		const int32 InRandomSeed,
+		const float InIsoValue,
+		const float InSpread,
+		TArray<FVector>& OutPoints) const
+	{
+		using namespace openvdb;
+
+		OutPoints.Empty();
+
+		if (this->FloatGrid->empty() ||
+			!(this->FloatGrid->getGridClass() == GRID_LEVEL_SET || this->FloatGrid->getGridClass() == GRID_FOG_VOLUME))
+		{
+			return;
+		}
+
+		openvdb::Grid<BoolTreeT>::Ptr InteriorMaskPtr = GetVolumeForScatter(this->FloatGrid, InIsoValue);
+
+		FRandomStream RandStream(InRandomSeed);
+		float NumPointsPerVoxel = InMinNumberOfPointsPerVoxel;
+		if (InMaxNumberOfPointsPerVoxel > InMinNumberOfPointsPerVoxel)
+		{
+			NumPointsPerVoxel = RandStream.FRandRange(InMinNumberOfPointsPerVoxel, InMaxNumberOfPointsPerVoxel);
+		}
+
+		using RandGen = std::mersenne_twister_engine<uint32_t, 32, 351, 175, 19,
+			0xccab8ee7, 11, 0xffffffff, 7, 0x31b6ab00, 15, 0xffe50000, 17, 1812433253>; // mt11213b
+		RandGen MTRand(/*seed=*/InRandomSeed);
+
+		std::vector<Vec3R> InstancePoints;
+		//			InstancePoints.reserve(NumPoints);
+		tools::v2s_internal::PointAccessor PointAcc(InstancePoints);
+
+		tools::NonUniformPointScatter<tools::v2s_internal::PointAccessor, RandGen> NonUniformScatter(PointAcc, NumPointsPerVoxel, MTRand, double(InSpread));
+		NonUniformScatter(*InteriorMaskPtr);
+
+		for (std::vector<Vec3R>::iterator It = InstancePoints.begin(); It != InstancePoints.end(); It++)
+		{
+			const FVector Position = FVector(It->x(), It->y(), It->z());
+			OutPoints.Add(Position);
+		}
+		return;
+	}
+
+	/* --------------------------------------------------------------------------------------------------------------- */
+
+	void FDataflowFloatVolumeImpl::VolumeSample(const TArray<FVector>& InPoints, TArray<float>& OutValues) const
+	{
+		using namespace openvdb;
+
+		if (this->FloatGrid->empty() ||
+			!(this->FloatGrid->getGridClass() == GRID_LEVEL_SET || this->FloatGrid->getGridClass() == GRID_FOG_VOLUME))
+		{
+			return;
+		}
+
+		openvdb::FloatGrid::ConstAccessor Accessor = this->FloatGrid->getConstAccessor();
+
+		openvdb::tools::GridSampler<openvdb::FloatGrid::ConstAccessor, openvdb::tools::BoxSampler>
+			fastSampler(Accessor, this->FloatGrid->transform());
+
+		int32 Idx = 0;
+		for (const FVector& Point : InPoints)
+		{
+			OutValues[Idx++] = fastSampler.wsSample(openvdb::Vec3R(Point.X, Point.Y, Point.Z));
+		}
+
+		return;
+	}
+
+	/* --------------------------------------------------------------------------------------------------------------- */
+
+	TSharedPtr<FDataflowFloatVectorVolumeImpl> FDataflowFloatVolumeImpl::ComputeGradient(const EDataflowVolumeAnalysisOutputName OutputName, const FString& CustomName, TSharedPtr<const FDataflowBoolVolumeImpl> InMaskVolumeImpl) const
+	{
+		using namespace openvdb;
+
+		constexpr openvdb::VecType VectorGridType = openvdb::VEC_COVARIANT;
+		const FString GridName = "Gradient";
+
+		TSharedPtr<FDataflowFloatVectorVolumeImpl> FloatVectorVolumeImpl = MakeShared<FDataflowFloatVectorVolumeImpl>();
+
+		if (this->FloatGrid->empty())
+		{
+			FloatVectorVolumeImpl->SetGrid(openvdb::VectorGrid::Ptr(), VectorGridType, GridName);
+			return FloatVectorVolumeImpl;
+		}
+
+		FString NewGridName;
+		if (OutputName == EDataflowVolumeAnalysisOutputName::AppendOperationName)
+		{
+			const FString CurrentGridName(this->FloatGrid->getName().c_str());
+			NewGridName = CurrentGridName + "_" + GridName;
+		}
+		else if (OutputName == EDataflowVolumeAnalysisOutputName::CustomName)
+		{
+			NewGridName = CustomName;
+		}
+
+		// @todo(gmelich): Implement using interrupter
+//			UE::DataflowVolumeUtils::FInterrupter NullInterrupter;
+
+		openvdb::VectorGrid::Ptr GradientVectorGrid;
+
+		if (InMaskVolumeImpl->GetGrid())
+		{
+			openvdb::BoolGrid RegionMask;
+
+			RegionMask.setTransform(this->FloatGrid->transform().copy());
+
+			openvdb::tools::resampleToMatch<openvdb::tools::PointSampler>(*InMaskVolumeImpl->GetGrid().get(), RegionMask);
+
+			openvdb::tools::Gradient<openvdb::FloatGrid, openvdb::BoolGrid> GradientTool(*this->FloatGrid.get(), RegionMask);
+			GradientVectorGrid = GradientTool.process(/*threaded*/true);
+		}
+		else
+		{
+			openvdb::tools::Gradient<openvdb::FloatGrid> GradientTool(*this->FloatGrid.get());
+			GradientVectorGrid = GradientTool.process(/*threaded*/true);
+		}
+
+		FloatVectorVolumeImpl->SetGrid(GradientVectorGrid, VectorGridType, NewGridName);
+
+		return FloatVectorVolumeImpl;
+	}
+
+	/* --------------------------------------------------------------------------------------------------------------- */
+
+	TSharedPtr<FDataflowFloatVolumeImpl> FDataflowFloatVolumeImpl::ComputeCurvature(const EDataflowVolumeAnalysisOutputName OutputName, const FString& CustomName, TSharedPtr<const FDataflowBoolVolumeImpl> InMaskVolumeImpl) const
+	{
+		using namespace openvdb;
+
+		constexpr openvdb::GridClass GridClass = openvdb::GRID_UNKNOWN;
+		const FString GridName = "Curvature";
+
+		TSharedPtr<FDataflowFloatVolumeImpl> FloatVolumeImpl = MakeShared<FDataflowFloatVolumeImpl>();
+
+		if (this->FloatGrid->empty())
+		{
+			FloatVolumeImpl->SetGrid(openvdb::FloatGrid::Ptr(), GridClass, GridName);
+			return FloatVolumeImpl;
+		}
+
+		FString NewGridName;
+		if (OutputName == EDataflowVolumeAnalysisOutputName::AppendOperationName)
+		{
+			const FString CurrentGridName(this->FloatGrid->getName().c_str());
+			NewGridName = CurrentGridName + "_" + GridName;
+		}
+		else if (OutputName == EDataflowVolumeAnalysisOutputName::CustomName)
+		{
+			NewGridName = CustomName;
+		}
+
+		// @todo(gmelich): Implement using interrupter
+//			UE::DataflowVolumeUtils::FInterrupter NullInterrupter;
+
+		openvdb::FloatGrid::Ptr CurvatureGrid;
+
+		if (InMaskVolumeImpl->GetGrid())
+		{
+			openvdb::BoolGrid RegionMask;
+
+			RegionMask.setTransform(this->FloatGrid->transform().copy());
+
+			openvdb::tools::resampleToMatch<openvdb::tools::PointSampler>(*InMaskVolumeImpl->GetGrid().get(), RegionMask);
+
+			openvdb::tools::MeanCurvature<openvdb::FloatGrid, openvdb::BoolGrid> CurvatureTool(*this->FloatGrid.get(), RegionMask);
+			CurvatureGrid = CurvatureTool.process(/*threaded*/true);
+		}
+		else
+		{
+			openvdb::tools::MeanCurvature<openvdb::FloatGrid> CurvatureTool(*this->FloatGrid.get());
+			CurvatureGrid = CurvatureTool.process(/*threaded*/true);
+		}
+
+		FloatVolumeImpl->SetGrid(CurvatureGrid, GridClass, NewGridName);
+
+		return FloatVolumeImpl;
+	}
+
+	/* --------------------------------------------------------------------------------------------------------------- */
+
+	TSharedPtr<FDataflowFloatVolumeImpl> FDataflowFloatVolumeImpl::ComputeLaplacian(const EDataflowVolumeAnalysisOutputName OutputName, const FString& CustomName, TSharedPtr<const FDataflowBoolVolumeImpl> InMaskVolumeImpl) const
+	{
+		using namespace openvdb;
+
+		constexpr openvdb::GridClass GridClass = openvdb::GRID_UNKNOWN;
+		const FString GridName = "Laplacian";
+
+		TSharedPtr<FDataflowFloatVolumeImpl> FloatVolumeImpl = MakeShared<FDataflowFloatVolumeImpl>();
+
+		if (this->FloatGrid->empty())
+		{
+			FloatVolumeImpl->SetGrid(openvdb::FloatGrid::Ptr(), GridClass, GridName);
+			return FloatVolumeImpl;
+		}
+
+		FString NewGridName;
+		if (OutputName == EDataflowVolumeAnalysisOutputName::AppendOperationName)
+		{
+			const FString CurrentGridName(this->FloatGrid->getName().c_str());
+			NewGridName = CurrentGridName + "_" + GridName;
+		}
+		else if (OutputName == EDataflowVolumeAnalysisOutputName::CustomName)
+		{
+			NewGridName = CustomName;
+		}
+
+		// @todo(gmelich): Implement using interrupter
+//			UE::DataflowVolumeUtils::FInterrupter NullInterrupter;
+
+		openvdb::FloatGrid::Ptr LaplacianGrid;
+
+		if (InMaskVolumeImpl->GetGrid())
+		{
+			openvdb::BoolGrid RegionMask;
+
+			RegionMask.setTransform(this->FloatGrid->transform().copy());
+
+			openvdb::tools::resampleToMatch<openvdb::tools::PointSampler>(*InMaskVolumeImpl->GetGrid().get(), RegionMask);
+
+			openvdb::tools::Laplacian<openvdb::FloatGrid, openvdb::BoolGrid> LaplacianTool(*this->FloatGrid.get(), RegionMask);
+			LaplacianGrid = LaplacianTool.process(/*threaded*/true);
+		}
+		else
+		{
+			openvdb::tools::Laplacian<openvdb::FloatGrid> LaplacianTool(*this->FloatGrid.get());
+			LaplacianGrid = LaplacianTool.process(/*threaded*/true);
+		}
+
+		FloatVolumeImpl->SetGrid(LaplacianGrid, GridClass, NewGridName);
+
+		return FloatVolumeImpl;
+	}
+
+	/* --------------------------------------------------------------------------------------------------------------- */
+
+	TSharedPtr<FDataflowFloatVectorVolumeImpl> FDataflowFloatVolumeImpl::ComputeClosestPoint(const EDataflowVolumeAnalysisOutputName OutputName, const FString& CustomName, TSharedPtr<const FDataflowBoolVolumeImpl> InMaskVolumeImpl) const
+	{
+		using namespace openvdb;
+
+		constexpr openvdb::VecType VectorGridType = openvdb::VEC_CONTRAVARIANT_ABSOLUTE;
+		const FString GridName = "ClosestPoint";
+
+		TSharedPtr<FDataflowFloatVectorVolumeImpl> FloatVectorVolumeImpl = MakeShared<FDataflowFloatVectorVolumeImpl>();
+
+		if (this->FloatGrid->empty())
+		{
+			FloatVectorVolumeImpl->SetGrid(openvdb::VectorGrid::Ptr(), VectorGridType, GridName);
+			return FloatVectorVolumeImpl;
+		}
+
+		FString NewGridName;
+		if (OutputName == EDataflowVolumeAnalysisOutputName::AppendOperationName)
+		{
+			const FString CurrentGridName(this->FloatGrid->getName().c_str());
+			NewGridName = CurrentGridName + "_" + GridName;
+		}
+		else if (OutputName == EDataflowVolumeAnalysisOutputName::CustomName)
+		{
+			NewGridName = CustomName;
+		}
+
+		// @todo(gmelich): Implement using interrupter
+
+//			UE::DataflowVolumeUtils::FInterrupter NullInterrupter;
+
+		openvdb::VectorGrid::Ptr CptVectorGrid;
+
+		if (InMaskVolumeImpl->GetGrid())
+		{
+			openvdb::BoolGrid RegionMask;
+
+			RegionMask.setTransform(this->FloatGrid->transform().copy());
+
+			openvdb::tools::resampleToMatch<openvdb::tools::PointSampler>(*InMaskVolumeImpl->GetGrid().get(), RegionMask);
+
+			openvdb::tools::Cpt<openvdb::FloatGrid, openvdb::BoolGrid> CptTool(*this->FloatGrid.get(), RegionMask);
+			CptVectorGrid = CptTool.process(/*threaded*/true);
+		}
+		else
+		{
+			openvdb::tools::Cpt<openvdb::FloatGrid> CptTool(*this->FloatGrid.get());
+			CptVectorGrid = CptTool.process(/*threaded*/true);
+		}
+
+		FloatVectorVolumeImpl->SetGrid(CptVectorGrid, VectorGridType, NewGridName);
+
+		return FloatVectorVolumeImpl;
+	}
+
+	/* --------------------------------------------------------------------------------------------------------------- */
+
+	// Functor to MultiplyAdd all active values
+
+	struct Op {
+		float Frac, Offset;
+		Op(float InFrac, float InOffset) : Frac(InFrac), Offset(InOffset) {}
+
+		inline void operator()(const openvdb::FloatGrid::ValueOnIter& iter) const
+		{
+			iter.setValue(*iter * Frac + Offset);
+		}
+	};
+
+	static void Combine(const float InScale, const float InOffset, const openvdb::FloatGrid::Ptr& InGrid, openvdb::FloatGrid::Ptr& OutGrid)
+	{
+		if (FMath::IsNearlyZero(InScale - 1.f, 1e-6) && FMath::IsNearlyZero(InOffset, 1e-6))
+		{
+			OutGrid = InGrid->deepCopy();
+		}
+		else
+		{
+			OutGrid = openvdb::FloatGrid::create(*InGrid);
+
+			float BackgroundValue = InScale * InGrid->background() + InOffset;
+			openvdb::tools::changeBackground(OutGrid->tree(), BackgroundValue);
+			openvdb::tools::foreach(OutGrid->beginValueOn(), Op(InScale, InOffset));
+		}
+	}
+/*
+	void ResampleGrids(openvdb::FloatGrid::Ptr FloatGridA, 
+		openvdb::FloatGrid::Ptr FloatGridB, 
+		const EDataflowVolumeSDFCombineOperation InOperation,
+		const EDataflowVolumeSDFCombineResample InResample,
+		const EDataflowVolumeSDFCombineInterpolation InInterpolation)
+	{
+		const bool NeedA = InOperation != EDataflowVolumeSDFCombineOperation::CopyB;
+		const bool NeedB = InOperation != EDataflowVolumeSDFCombineOperation::CopyA && InOperation != EDataflowVolumeSDFCombineOperation::InvertA;
+		const bool NeedBoth = NeedA && NeedB;
+		const bool NeedLevelSets = InOperation == EDataflowVolumeSDFCombineOperation::SDFUnion || 
+			InOperation == EDataflowVolumeSDFCombineOperation::SDFIntersect ||
+			InOperation == EDataflowVolumeSDFCombineOperation::SDFDifference;
+
+		int32 ResampleWhich = (int32)EDataflowVolumeSDFCombineResample::Off;
+
+		if (!FloatGridA->empty() && !FloatGridB->empty())
+		{
+			if (InResample == EDataflowVolumeSDFCombineResample::HiresMatchLores || InResample == EDataflowVolumeSDFCombineResample::LoresMatchHires)
+			{
+				const openvdb::Vec3d aVoxSize = FloatGridA->voxelSize(), bVoxSize = FloatGridB->voxelSize();
+				const double aVoxVol = aVoxSize[0] * aVoxSize[1] * aVoxSize[2], bVoxVol = bVoxSize[0] * bVoxSize[1] * bVoxSize[2];
+				ResampleWhich = ((aVoxVol > bVoxVol && InResample == EDataflowVolumeSDFCombineResample::LoresMatchHires)
+					|| (aVoxVol < bVoxVol && InResample == EDataflowVolumeSDFCombineResample::HiresMatchLores))
+					? (int32)EDataflowVolumeSDFCombineResample::AMatchB : (int32)EDataflowVolumeSDFCombineResample::BMatchA;
+			}
+			else 
+			{
+				ResampleWhich = (int32)InResample;
+			}
+
+			openvdb::FloatGrid::Ptr aBaseGrid, bBaseGrid;
+
+			if (FloatGridA->constTransform() != FloatGridB->constTransform()) {
+				// If the A and B grid transforms don't match, one of the grids
+				// should be resampled into the other's index space.
+				if (InResample == EDataflowVolumeSDFCombineResample::Off) {
+					if (NeedBoth) {
+						// Resampling is disabled.  Just log a warning.
+//						std::ostringstream ostr;
+//						ostr << aGridName << " and " << bGridName << " transforms don't match";
+//						self->addWarning(SOP_MESSAGE, ostr.str().c_str());
+					}
+				}
+				else {
+					if (NeedA && ResampleWhich == EDataflowVolumeSDFCombineResample::AMatchB) {
+						// Resample grid A into grid B's index space.
+						aBaseGrid = this->resampleToMatch(*aGrid, *bGrid, samplingOrder);
+						aGrid = static_cast<const AGridT*>(aBaseGrid.get());
+					}
+					else if (needB && resampleWhich == RESAMPLE_B) {
+						// Resample grid B into grid A's index space.
+						bBaseGrid = this->resampleToMatch(*bGrid, *aGrid, samplingOrder);
+						bGrid = static_cast<const BGridT*>(bBaseGrid.get());
+					}
+				}
+			}
+
+
+
+
+
+
+
+
+
+
+
+
+
+
+
+
+
+
+		}
+	}
+
+
+
+
+
+
+
+
+
+
+
+
+
+
+
+
+
+
+
+
+
+
+
+
+
+
+
+
+
+
+
+
+
+
+	//enum ResampleMode {
+	//	RESAMPLE_OFF,    // don't auto-resample grids
+	//	RESAMPLE_B,      // resample B to match A
+	//	RESAMPLE_A,      // resample A to match B
+	//	RESAMPLE_HI_RES, // resample higher-res grid to match lower-res
+	//	RESAMPLE_LO_RES  // resample lower-res grid to match higher-res
+	//};
+	//enum { RESAMPLE_MODE_FIRST = RESAMPLE_OFF, RESAMPLE_MODE_LAST = RESAMPLE_LO_RES };
+
+
+*/
+
+
+
+	TSharedPtr<FDataflowFloatVolumeImpl> FDataflowFloatVolumeImpl::VolumeCombine(
+		TSharedPtr<const FDataflowFloatVolumeImpl> InVolumeBImpl,
+		const EDataflowVolumeSDFCombineOperation InOperation,
+		const float InMultiplierA,
+		const float InMultiplierB,
+		const EDataflowVolumeSDFCombineResample InResample,
+		const EDataflowVolumeSDFCombineInterpolation InInterpolation,
+		const bool InPruneDegenerateTiles,
+		const float InPruneTol) const
+	{
+		constexpr openvdb::GridClass GridClass = openvdb::GRID_UNKNOWN;
+
+		using namespace openvdb;
+
+		TSharedPtr<FDataflowFloatVolumeImpl> FloatVolumeImpl = MakeShared<FDataflowFloatVolumeImpl>();
+
+		if (this->FloatGrid->empty() || InVolumeBImpl->FloatGrid->empty())
+		{
+			FloatVolumeImpl->SetGrid(openvdb::FloatGrid::Ptr(), GridClass, "Combine");
+			return FloatVolumeImpl;
+		}
+
+		// @todo(gmelich): Implement resmapling the grids
+//		ResampleGrids(this->FloatGrid, InVolumeB.FloatGrid, InOperation, InResample, InInterpolation);
+
+		openvdb::FloatGrid::Ptr ResultGrid, TempGrid;
+
+		switch (InOperation) {
+		case EDataflowVolumeSDFCombineOperation::SDFIntersect:
+			Combine(InMultiplierA, 0.f, this->FloatGrid, ResultGrid);
+			Combine(InMultiplierB, 0.f, InVolumeBImpl->FloatGrid, TempGrid);
+			openvdb::tools::csgIntersection(*ResultGrid, *TempGrid, /*prune*/true);
+
+			break;
+
+		case EDataflowVolumeSDFCombineOperation::SDFUnion:
+			Combine(InMultiplierA, 0.f, this->FloatGrid, ResultGrid);
+			Combine(InMultiplierB, 0.f, InVolumeBImpl->FloatGrid, TempGrid);
+			openvdb::tools::csgUnion(*ResultGrid, *TempGrid, /*prune*/true);
+
+			break;
+
+		case EDataflowVolumeSDFCombineOperation::SDFDifference:
+			Combine(InMultiplierA, 0.f, this->FloatGrid, ResultGrid);
+			Combine(InMultiplierB, 0.f, InVolumeBImpl->FloatGrid, TempGrid);
+			openvdb::tools::csgDifference(*ResultGrid, *TempGrid, /*prune*/true);
+			
+			break;
+
+		default:
+			break;
+		}
+
+		openvdb::tools::prune(ResultGrid->tree(), InPruneTol);
+
+		FloatVolumeImpl->SetGrid(ResultGrid, openvdb::GRID_LEVEL_SET);
+
+		return FloatVolumeImpl;
+	}
+
+	/* --------------------------------------------------------------------------------------------------------------- */
+
+
+
+
+
+
+}
\ No newline at end of file
diff --git a/Engine/Plugins/Experimental/DataflowVolumetricPlugin/Source/DataflowVolumeCore/Private/Dataflow/DataflowFloatVolumeImpl.h b/Engine/Plugins/Experimental/DataflowVolumetricPlugin/Source/DataflowVolumeCore/Private/Dataflow/DataflowFloatVolumeImpl.h
new file mode 100644
--- /dev/null
+++ b/Engine/Plugins/Experimental/DataflowVolumetricPlugin/Source/DataflowVolumeCore/Private/Dataflow/DataflowFloatVolumeImpl.h
@@ -0,0 +1,144 @@
+// Copyright Epic Games, Inc. All Rights Reserved.
+
+#pragma once
+
+#include "Dataflow/DataflowVolumeImpl.h"
+#include "Dataflow/DataflowFloatVectorVolumeImpl.h"
+#include "Dataflow/OpenVDB.h"
+#include "GeometryCollection/ManagedArrayCollection.h"
+#include "MeshDescription.h"
+#include "Dataflow/DataflowVolumeNodeEnums.h"
+
+struct FDataflowBoolVolume;
+
+namespace UE::DataflowVolume::Private
+{
+	/* --------------------------------------------------------------------------------------------------------------- */
+	/* FDataflowFloatVolumeImpl */
+	/* --------------------------------------------------------------------------------------------------------------- */
+
+	class FDataflowVolume;
+	class FDataflowIntVolumeImpl;
+	class FDataflowFloatVectorVolumeImpl;
+	class FDataflowBoolVolumeImpl;
+
+	class FDataflowFloatVolumeImpl : public FDataflowVolumeImpl
+	{
+	private:
+		openvdb::FloatGrid::Ptr FloatGrid = nullptr;
+
+	public:
+		openvdb::FloatGrid::Ptr GetGrid() const;
+
+		void SetGrid(const openvdb::FloatGrid::Ptr& InFloatGrid, const openvdb::GridClass InGridClass, const FString& InGridName = TEXT("LevelSet"));
+
+		FString GetGridType() const;
+
+		FString GetGridClass() const;
+
+		virtual FString VolumeInfo() const override;
+
+		static TSharedPtr<FDataflowFloatVolumeImpl> CreateSphereSDF(
+			const float InVoxelSize,
+			const float InRadius,
+			const FVector& InCenter);
+
+		static TSharedPtr<FDataflowFloatVolumeImpl> CreateCubeSDF(
+			const float InScale, 
+			const FVector& InCenter, 
+			const float InVoxelSize);
+
+		static TSharedPtr<FDataflowFloatVolumeImpl> CreatePlatonicSolidSDF(
+			const int32 InFaceCount, 
+			const float InScale, 
+			const FVector& InCenter, 
+			const float InVoxelSize);
+
+		static TSharedPtr<FDataflowFloatVolumeImpl> CreateVolumeFromMeshDescription(
+			const FMeshDescription& InMeshDescription,
+			const FTransform& InTransform,
+			const float InVoxelSize,
+			const TSharedPtr<const FDataflowVolumeImpl>& InReferenceVolumeImpl,
+			const EDataflowVolumeOutputType InOutputType,
+			const FString& InGridName,
+			const bool InUseWorldSpaceUnits,
+			const float InExteriorBand,
+			const float InInteriorBand,
+			const int32 InExteriorBandVoxels,
+			const int32 InInteriorBandVoxels,
+			const bool bFillInterior,
+			const bool bPreserveHoles,
+			const float InIsoValue,
+			TSharedPtr<FDataflowIntVolumeImpl>& OutFaceIndexVolume);
+
+		void VolumeToSpheres(const int32 InMinSphereCount,
+			const int32 InMaxSphereCount,
+			const bool InOverlapping,
+			const float InMinRadius,
+			float InMaxRadius,
+			const float InIsoValue,
+			const int32 InInstanceCount,
+			TArray<FVector>& SphereCenters,
+			TArray<float>& SphereRadii) const;
+
+		TSharedPtr<FDataflowFloatVolumeImpl> CovertSDFToFogVolume(
+			const bool InPruneTolerance,
+			const float InTolerance,
+			const bool InFloodFillOutput,
+			const bool InActivateInteriror) const;
+
+		TSharedPtr<FDataflowFloatVolumeImpl> CovertFogVolumeToSDF(
+			const float InFogIsoValue,
+			const bool InPruneTolerance,
+			const float InTolerance,
+			const bool InFloodFillOutput,
+			const bool InActivateInteriror) const;
+
+		void CovertSDFToCollection(const float InIsoValue, const float InAdaptivity, FManagedArrayCollection& OutCollection) const;
+
+		void ConvertVolumeToMeshDescription(const float InIsoValue, const float InAdaptivity, FMeshDescription& OutMeshDescription) const;
+
+		virtual void GetActiveVoxels(TArray<FBox>& OutActiveVoxels) const override;
+
+		void UniformVolumeScatter(
+			const int32 InMinNumberOfPoints,
+			const int32 InMaxNumberOfPoints,
+			const int32 InRandomSeed,
+			const float InIsoValue,
+			const float InSpread,
+			TArray<FVector>& OutPoints) const;
+
+		void DenseUniformVolumeScatter(
+			const float InMinNumberOfPointsPerVoxel,
+			const float InMaxNumberOfPointsPerVoxel,
+			const int32 InRandomSeed,
+			const float InIsoValue,
+			const float InSpread,
+			TArray<FVector>& OutPoints) const;
+
+		void NonUniformVolumeScatter(
+			const float InMinNumberOfPointsPerVoxel,
+			const float InMaxNumberOfPointsPerVoxel,
+			const int32 InRandomSeed,
+			const float InIsoValue,
+			const float InSpread,
+			TArray<FVector>& OutPoints) const;
+
+		void VolumeSample(const TArray<FVector>& InPoints, TArray<float>& OutValues) const;
+
+		TSharedPtr<FDataflowFloatVectorVolumeImpl> ComputeGradient(const EDataflowVolumeAnalysisOutputName OutputName, const FString& CustomName, TSharedPtr<const FDataflowBoolVolumeImpl> InMaskVolumeImpl) const;
+		TSharedPtr<FDataflowFloatVolumeImpl> ComputeCurvature(const EDataflowVolumeAnalysisOutputName OutputName, const FString& CustomName, TSharedPtr<const FDataflowBoolVolumeImpl> InMaskVolumeImpl) const;
+		TSharedPtr<FDataflowFloatVolumeImpl> ComputeLaplacian(const EDataflowVolumeAnalysisOutputName OutputName, const FString& CustomName, TSharedPtr<const FDataflowBoolVolumeImpl> InMaskVolumeImpl) const;
+		TSharedPtr<FDataflowFloatVectorVolumeImpl> ComputeClosestPoint(const EDataflowVolumeAnalysisOutputName OutputName, const FString& CustomName, TSharedPtr<const FDataflowBoolVolumeImpl> InMaskVolumeImpl) const;
+
+		TSharedPtr<FDataflowFloatVolumeImpl> VolumeCombine(
+			TSharedPtr<const FDataflowFloatVolumeImpl> InVolumeBImpl,
+			const EDataflowVolumeSDFCombineOperation InOperation,
+			const float InMultiplierA,
+			const float InMultiplierB,
+			const EDataflowVolumeSDFCombineResample InResample,
+			const EDataflowVolumeSDFCombineInterpolation InInterpolation,
+			const bool InPruneDegenerateTiles,
+			const float InPruneTol) const;
+	};
+}
\ No newline at end of file
diff --git a/Engine/Plugins/Experimental/DataflowVolumetricPlugin/Source/DataflowVolumeCore/Private/Dataflow/DataflowIntVolume.cpp b/Engine/Plugins/Experimental/DataflowVolumetricPlugin/Source/DataflowVolumeCore/Private/Dataflow/DataflowIntVolume.cpp
new file mode 100644
--- /dev/null
+++ b/Engine/Plugins/Experimental/DataflowVolumetricPlugin/Source/DataflowVolumeCore/Private/Dataflow/DataflowIntVolume.cpp
@@ -0,0 +1,71 @@
+// Copyright Epic Games, Inc. All Rights Reserved.
+
+#include "Dataflow/DataflowIntVolume.h"
+#include "Dataflow/DataflowIntVolumeImpl.h"
+//#include "Dataflow/OpenVDB.h"
+//#include "Dataflow/DataflowVolumeUtils.h"
+//#include <limits>
+
+/* --------------------------------------------------------------------------------------------------------------- */
+/* FDataflowIntVolume Implementation */
+/* --------------------------------------------------------------------------------------------------------------- */
+
+FDataflowIntVolume::FDataflowIntVolume(TSharedPtr<UE::DataflowVolume::Private::FDataflowIntVolumeImpl> InIntVolume)
+	: FDataflowVolume(InIntVolume, FDataflowIntVolume::TypeName)
+{
+	static_assert(sizeof(FDataflowVolume) == sizeof(*this));
+}
+
+TSharedPtr<const UE::DataflowVolume::Private::FDataflowIntVolumeImpl> FDataflowIntVolume::GetIntVolume() const
+{
+	return StaticCastSharedPtr<const UE::DataflowVolume::Private::FDataflowIntVolumeImpl>(Volume);
+}
+
+TSharedPtr<UE::DataflowVolume::Private::FDataflowIntVolumeImpl> FDataflowIntVolume::GetIntVolume()
+{
+	return StaticCastSharedPtr<UE::DataflowVolume::Private::FDataflowIntVolumeImpl>(Volume);
+}
+
+
+void FDataflowIntVolume::SetIntVolume(TSharedPtr<UE::DataflowVolume::Private::FDataflowIntVolumeImpl> InIntVolume)
+{
+	Volume = InIntVolume;
+}
+
+FVector FDataflowIntVolume::GetVoxelSize() const
+{
+	using namespace UE::DataflowVolume::Private;
+
+	if (TSharedPtr<const FDataflowIntVolumeImpl> IntVolume = GetIntVolume())
+	{
+		return IntVolume->GetVoxelSize();
+	}
+
+	return {};
+}
+
+FBox FDataflowIntVolume::GetVolumeBoundingBox() const
+{
+	using namespace UE::DataflowVolume::Private;
+
+	if (TSharedPtr<const FDataflowIntVolumeImpl> IntVolume = GetIntVolume())
+	{
+		return IntVolume->GetActiveVoxelBoundingBox();
+	}
+
+	return {};
+}
+
+void FDataflowIntVolume::VolumeSample(
+	const TArray<FVector>& InPoints,
+	TArray<int32>& OutValues) const
+{
+	using namespace UE::DataflowVolume::Private;
+
+	if (TSharedPtr<const FDataflowIntVolumeImpl> IntVolume = GetIntVolume())
+	{
+		IntVolume->VolumeSample(
+			InPoints,
+			OutValues);
+	}
+}
diff --git a/Engine/Plugins/Experimental/DataflowVolumetricPlugin/Source/DataflowVolumeCore/Private/Dataflow/DataflowIntVolumeImpl.cpp b/Engine/Plugins/Experimental/DataflowVolumetricPlugin/Source/DataflowVolumeCore/Private/Dataflow/DataflowIntVolumeImpl.cpp
new file mode 100644
--- /dev/null
+++ b/Engine/Plugins/Experimental/DataflowVolumetricPlugin/Source/DataflowVolumeCore/Private/Dataflow/DataflowIntVolumeImpl.cpp
@@ -0,0 +1,106 @@
+// Copyright Epic Games, Inc. All Rights Reserved.
+
+#include "Dataflow/DataflowIntVolumeImpl.h"
+//#include "Dataflow/OpenVDB.h"
+//#include "Dataflow/DataflowVolumeUtils.h"
+//#include <limits>
+
+namespace UE::DataflowVolume::Private
+{
+	/* --------------------------------------------------------------------------------------------------------------- */
+	/* FDataflowIntVolumeImpl */
+	/* --------------------------------------------------------------------------------------------------------------- */
+
+	void FDataflowIntVolumeImpl::SetGrid(const openvdb::Int32Grid::Ptr& InIntGrid, const FString& InGridName)
+	{
+		FDataflowVolumeImpl::SetGrid(InIntGrid, InGridName);
+
+		IntGrid = InIntGrid;
+	}
+
+	openvdb::Int32Grid::Ptr FDataflowIntVolumeImpl::GetGrid() const
+	{
+		return IntGrid;
+	}
+
+	FString FDataflowIntVolumeImpl::GetGridType() const
+	{
+		if (IntGrid)
+		{
+			std::string GridType = IntGrid->gridType();
+			return FString(GridType.c_str());
+		}
+
+		return {};
+	}
+
+	FString FDataflowIntVolumeImpl::VolumeInfo() const
+	{
+		FString BaseString = FDataflowVolumeImpl::VolumeInfo();
+
+		FString OutputStr;
+		OutputStr.Appendf(TEXT("\n----------------------------------------\n"));
+		OutputStr += BaseString;
+
+		FString GridType = this->GetGridType();
+
+		OutputStr.Appendf(TEXT("Type: %s\n"),
+			*GridType);
+
+		OutputStr.Appendf(TEXT("----------------------------------------\n"));
+
+		return OutputStr;
+	}
+
+	void FDataflowIntVolumeImpl::VolumeSample(const TArray<FVector>& InPoints, TArray<int32>& OutValues) const
+	{
+		using namespace openvdb;
+
+		if (this->IntGrid->empty())
+		{
+			return;
+		}
+
+		openvdb::Int32Grid::ConstAccessor Accessor = this->IntGrid->getConstAccessor();
+
+		openvdb::tools::GridSampler<openvdb::Int32Grid::ConstAccessor, openvdb::tools::BoxSampler>
+			fastSampler(Accessor, this->IntGrid->transform());
+
+		int32 Idx = 0;
+		for (const FVector& Point : InPoints)
+		{
+			OutValues[Idx++] = fastSampler.wsSample(openvdb::Vec3R(Point.X, Point.Y, Point.Z));
+		}
+
+		return;
+	}
+
+	void FDataflowIntVolumeImpl::GetActiveVoxels(TArray<FBox>& OutActiveVoxels) const
+	{
+		const openvdb::Index64 ActiveVoxelCount = IntGrid->activeVoxelCount();
+		OutActiveVoxels.AddUninitialized(ActiveVoxelCount);
+
+		openvdb::math::Transform T = this->IntGrid->transform();
+
+		int32 Idx = 0;
+		for (openvdb::Int32Grid::ValueOnCIter iter = this->IntGrid->cbeginValueOn(); iter.test(); ++iter)
+		{
+			if (iter.isVoxelValue())
+			{
+				openvdb::CoordBBox BBox;
+				iter.getBoundingBox(BBox);
+
+				const openvdb::Vec3d BBoxMin(BBox.min().x() - 0.5, BBox.min().y() - 0.5, BBox.min().z() - 0.5);
+				const openvdb::Vec3d BBoxMax(BBox.max().x() + 0.5, BBox.max().y() + 0.5, BBox.max().z() + 0.5);
+
+				const openvdb::Vec3d BBoxMinWS = IntGrid->transform().indexToWorld(BBoxMin);
+				const openvdb::Vec3d BBoxMaxWS = IntGrid->transform().indexToWorld(BBoxMax);
+
+				FVector Min(BBoxMinWS.x(), BBoxMinWS.y(), BBoxMinWS.z());
+				FVector Max(BBoxMaxWS.x(), BBoxMaxWS.y(), BBoxMaxWS.z());
+
+				OutActiveVoxels[Idx++] = FBox(Min, Max);
+			}
+		}
+	}
+}
diff --git a/Engine/Plugins/Experimental/DataflowVolumetricPlugin/Source/DataflowVolumeCore/Private/Dataflow/DataflowIntVolumeImpl.h b/Engine/Plugins/Experimental/DataflowVolumetricPlugin/Source/DataflowVolumeCore/Private/Dataflow/DataflowIntVolumeImpl.h
new file mode 100644
--- /dev/null
+++ b/Engine/Plugins/Experimental/DataflowVolumetricPlugin/Source/DataflowVolumeCore/Private/Dataflow/DataflowIntVolumeImpl.h
@@ -0,0 +1,32 @@
+// Copyright Epic Games, Inc. All Rights Reserved.
+
+#pragma once
+
+#include "Dataflow/DataflowVolumeImpl.h"
+#include "Dataflow/OpenVDB.h"
+
+namespace UE::DataflowVolume::Private
+{
+	/* --------------------------------------------------------------------------------------------------------------- */
+	/* FDataflowIntVolumeImpl */
+	/* --------------------------------------------------------------------------------------------------------------- */
+
+	class FDataflowIntVolumeImpl : public FDataflowVolumeImpl
+	{
+	private:
+		openvdb::Int32Grid::Ptr IntGrid = nullptr;
+
+	public:
+		void SetGrid(const openvdb::Int32Grid::Ptr& InIntGrid, const FString& InGridName = "");
+
+		openvdb::Int32Grid::Ptr GetGrid() const;
+
+		FString GetGridType() const;
+
+		virtual FString VolumeInfo() const override;
+
+		void VolumeSample(const TArray<FVector>& InPoints, TArray<int32>& OutValues) const;
+
+		virtual void GetActiveVoxels(TArray<FBox>& OutActiveVoxels) const override;
+	};
+}
diff --git a/Engine/Plugins/Experimental/DataflowVolumetricPlugin/Source/DataflowVolumeCore/Private/Dataflow/DataflowVolume.cpp b/Engine/Plugins/Experimental/DataflowVolumetricPlugin/Source/DataflowVolumeCore/Private/Dataflow/DataflowVolume.cpp
new file mode 100644
--- /dev/null
+++ b/Engine/Plugins/Experimental/DataflowVolumetricPlugin/Source/DataflowVolumeCore/Private/Dataflow/DataflowVolume.cpp
@@ -0,0 +1,34 @@
+// Copyright Epic Games, Inc. All Rights Reserved.
+
+#include "Dataflow/DataflowVolume.h"
+#include "Dataflow/DataflowVolumeImpl.h"
+
+/* --------------------------------------------------------------------------------------------------------------- */
+/* FDataflowVolume Implementation */
+/* --------------------------------------------------------------------------------------------------------------- */
+
+FDataflowVolume::FDataflowVolume(TSharedPtr<UE::DataflowVolume::Private::FDataflowVolumeImpl> InVolume, const FName InType)
+	: Volume(InVolume)
+	, Type(InType)
+{}
+
+FString FDataflowVolume::VolumeInfo() const
+{
+	if (Volume)
+	{
+		return Volume->VolumeInfo();
+	}
+
+	return {};
+}
+
+void FDataflowVolume::GetActiveVoxels(TArray<FBox>& OutActiveVoxels) const
+{
+	if (Volume)
+	{
+		Volume->GetActiveVoxels(OutActiveVoxels);
+	}
+}
+
+
+
diff --git a/Engine/Plugins/Experimental/DataflowVolumetricPlugin/Source/DataflowVolumeCore/Private/Dataflow/DataflowVolumeAnyType.cpp b/Engine/Plugins/Experimental/DataflowVolumetricPlugin/Source/DataflowVolumeCore/Private/Dataflow/DataflowVolumeAnyType.cpp
new file mode 100644
--- /dev/null
+++ b/Engine/Plugins/Experimental/DataflowVolumetricPlugin/Source/DataflowVolumeCore/Private/Dataflow/DataflowVolumeAnyType.cpp
@@ -0,0 +1,12 @@
+// Copyright Epic Games, Inc. All Rights Reserved.
+
+#include "Dataflow/DataflowVolumeAnyType.h"
+#include "Dataflow/DataflowAnyTypeRegistry.h"
+
+namespace UE::Dataflow
+{
+	void RegisterVolumeAnyTypes()
+	{
+		UE_DATAFLOW_REGISTER_ANYTYPE(FDataflowVolumeTypes);
+	}
+}
\ No newline at end of file
diff --git a/Engine/Plugins/Experimental/DataflowVolumetricPlugin/Source/DataflowVolumeCore/Private/Dataflow/DataflowVolumeCoreModule.cpp b/Engine/Plugins/Experimental/DataflowVolumetricPlugin/Source/DataflowVolumeCore/Private/Dataflow/DataflowVolumeCoreModule.cpp
new file mode 100644
--- /dev/null
+++ b/Engine/Plugins/Experimental/DataflowVolumetricPlugin/Source/DataflowVolumeCore/Private/Dataflow/DataflowVolumeCoreModule.cpp
@@ -0,0 +1,24 @@
+// Copyright Epic Games, Inc. All Rights Reserved.
+
+#include "Dataflow/DataflowVolumeCoreModule.h"
+
+#include "CoreMinimal.h"
+#include "Dataflow/DataflowNodeFactory.h"
+#include "Modules/ModuleInterface.h"
+#include "Modules/ModuleManager.h"
+#include "Dataflow/OpenVDB.h"
+
+#define LOCTEXT_NAMESPACE "DataflowCore"
+
+void IDataflowVolumeCoreModule::StartupModule()
+{
+	// Initialize OpenVDB
+	openvdb::initialize();
+}
+
+void IDataflowVolumeCoreModule::ShutdownModule()
+{}
+
+IMPLEMENT_MODULE(IDataflowVolumeCoreModule, DataflowVolumeCore)
+
+#undef LOCTEXT_NAMESPACE
diff --git a/Engine/Plugins/Experimental/DataflowVolumetricPlugin/Source/DataflowVolumeCore/Private/Dataflow/DataflowVolumeCoreModule.h b/Engine/Plugins/Experimental/DataflowVolumetricPlugin/Source/DataflowVolumeCore/Private/Dataflow/DataflowVolumeCoreModule.h
new file mode 100644
--- /dev/null
+++ b/Engine/Plugins/Experimental/DataflowVolumetricPlugin/Source/DataflowVolumeCore/Private/Dataflow/DataflowVolumeCoreModule.h
@@ -0,0 +1,30 @@
+// Copyright Epic Games, Inc. All Rights Reserved.
+
+#pragma once
+
+#include "CoreMinimal.h"
+#include "Dataflow/DataflowNodeFactory.h"
+#include "Modules/ModuleInterface.h"
+#include "Modules/ModuleManager.h"
+
+
+/**
+ * The public interface to this module
+ */
+class IDataflowVolumeCoreModule : public IModuleInterface
+{
+
+public:
+
+	// IModuleInterface interface
+
+	DATAFLOWVOLUMECORE_API virtual void StartupModule() override;
+
+	DATAFLOWVOLUMECORE_API virtual void ShutdownModule() override;
+
+	virtual bool SupportsDynamicReloading() override
+	{
+		return false;
+	}
+
+};
diff --git a/Engine/Plugins/Experimental/DataflowVolumetricPlugin/Source/DataflowVolumeCore/Private/Dataflow/DataflowVolumeImpl.cpp b/Engine/Plugins/Experimental/DataflowVolumetricPlugin/Source/DataflowVolumeCore/Private/Dataflow/DataflowVolumeImpl.cpp
new file mode 100644
--- /dev/null
+++ b/Engine/Plugins/Experimental/DataflowVolumetricPlugin/Source/DataflowVolumeCore/Private/Dataflow/DataflowVolumeImpl.cpp
@@ -0,0 +1,141 @@
+// Copyright Epic Games, Inc. All Rights Reserved.
+
+#include "Dataflow/DataflowVolumeImpl.h"
+#include "Dataflow/OpenVDB.h"
+
+namespace UE::DataflowVolume::Private
+{
+	/* --------------------------------------------------------------------------------------------------------------- */
+	/* FDataflowVolumeImpl */
+	/* --------------------------------------------------------------------------------------------------------------- */
+
+	bool FDataflowVolumeImpl::IsGridEmpty() const
+	{
+		if (BaseGrid)
+		{
+			return BaseGrid->empty();
+		}
+
+		return false;
+	}
+
+	FVector FDataflowVolumeImpl::GetVoxelSize() const
+	{
+		if (BaseGrid)
+		{
+			openvdb::Vec3d VoxelSize = BaseGrid->voxelSize();
+			return FVector(VoxelSize.x(), VoxelSize.y(), VoxelSize.z());
+		}
+		
+		return FVector(0);
+	}
+
+	FString FDataflowVolumeImpl::GetGridName() const
+	{
+		if (BaseGrid)
+		{
+			std::string GridName = BaseGrid->getName();
+			return FString(GridName.c_str());
+		}
+		
+		return {};
+	}
+
+	int32 FDataflowVolumeImpl::GetActiveVoxelCount() const
+	{
+		if (BaseGrid)
+		{
+			openvdb::Index64 ActiveVoxelCount = BaseGrid->activeVoxelCount();
+			return (int32)ActiveVoxelCount;
+		}
+
+		return 0;
+	}
+
+	FIntVector3 FDataflowVolumeImpl::GetActiveVoxelDim() const
+	{
+		if (BaseGrid)
+		{
+			return FIntVector3(
+				BaseGrid->evalActiveVoxelDim().x(),
+				BaseGrid->evalActiveVoxelDim().y(),
+				BaseGrid->evalActiveVoxelDim().z());
+		}
+
+		return {};
+	}
+
+	FBox FDataflowVolumeImpl::GetActiveVoxelBoundingBox() const
+	{
+		if (BaseGrid)
+		{
+			openvdb::CoordBBox BBox = BaseGrid->evalActiveVoxelBoundingBox();
+			openvdb::BBoxd BBoxInWorldSpace = BaseGrid->transform().indexToWorld(BBox);
+
+			FVector Min(BBoxInWorldSpace.min().x(), BBoxInWorldSpace.min().y(), BBoxInWorldSpace.min().z());
+			FVector Max(BBoxInWorldSpace.max().x(), BBoxInWorldSpace.max().y(), BBoxInWorldSpace.max().z());
+
+			return FBox(Min, Max);
+		}
+
+		return {};
+	}
+
+	FTransform FDataflowVolumeImpl::GetGridTransform() const
+	{
+		if (BaseGrid)
+		{
+			openvdb::math::Transform GridTransform = BaseGrid->transform();
+			openvdb::Mat4R xform = GridTransform.baseMap()->getAffineMap()->getMat4();
+
+			FMatrix Matrix = FMatrix::Identity;
+
+			for (int32 IdxRow = 0; IdxRow < 4; ++IdxRow)
+			{
+				for (int32 IdxCol = 0; IdxCol < 4; ++IdxCol)
+				{
+					Matrix.M[IdxCol][IdxRow] = xform[IdxCol][IdxRow];
+				}
+			}
+
+			FTransform Transform(Matrix);
+			return Transform;
+		}
+
+		return {};
+	}
+
+	FString FDataflowVolumeImpl::VolumeInfo() const
+	{
+		FString GridName = this->GetGridName();
+		FIntVector3 GridDimensions = this->GetActiveVoxelDim();
+		int32 ActiveVoxels = this->GetActiveVoxelCount();
+		FVector VoxelSize = this->GetVoxelSize();
+		FBox BoundingBox = this->GetActiveVoxelBoundingBox();
+		FTransform Transform = this->GetGridTransform();
+
+		FString OutputStr;
+		OutputStr.Appendf(TEXT("Grid Name: %s\nDimensions: [%d,%d,%d]\nActive Voxels: %s\nVoxel Size: %.3f\nBounding Box: [%3.2f,%3.2f,%3.2f][%3.2f,%3.2f,%3.2f]\nTransform: %3.2f,%3.2f,%3.2f|%3.2f,%3.2f,%3.2f|%3.2f,%3.2f,%3.2f\n"),
+			*GridName,
+			GridDimensions.X, GridDimensions.Y, GridDimensions.Z,
+			*FString::FormatAsNumber(ActiveVoxels),
+			VoxelSize.X,
+			BoundingBox.Min.X, BoundingBox.Min.Y, BoundingBox.Min.Z,
+			BoundingBox.Max.X, BoundingBox.Max.Y, BoundingBox.Max.Z,
+			Transform.GetTranslation().X, Transform.GetTranslation().Y, Transform.GetTranslation().Z,
+			Transform.GetRotation().Euler().X, Transform.GetRotation().Euler().Y, Transform.GetRotation().Euler().Z,
+			Transform.GetScale3D().X, Transform.GetScale3D().Y, Transform.GetScale3D().Z);
+
+		return OutputStr;
+	}
+	
+	void FDataflowVolumeImpl::SetGrid(const openvdb::GridBase::Ptr& InGrid, const FString& InGridName)
+	{
+		BaseGrid = InGrid;
+
+		std::string GridNameString(TCHAR_TO_UTF8(*InGridName));
+		BaseGrid->setName(GridNameString);
+
+		BaseGrid->setCreator("UE5_Dataflow");
+	}
+}
\ No newline at end of file
diff --git a/Engine/Plugins/Experimental/DataflowVolumetricPlugin/Source/DataflowVolumeCore/Private/Dataflow/DataflowVolumeImpl.h b/Engine/Plugins/Experimental/DataflowVolumetricPlugin/Source/DataflowVolumeCore/Private/Dataflow/DataflowVolumeImpl.h
new file mode 100644
--- /dev/null
+++ b/Engine/Plugins/Experimental/DataflowVolumetricPlugin/Source/DataflowVolumeCore/Private/Dataflow/DataflowVolumeImpl.h
@@ -0,0 +1,44 @@
+// Copyright Epic Games, Inc. All Rights Reserved.
+
+#pragma once
+
+#include "Dataflow/OpenVDB.h"
+#include "Math/UnrealMath.h"
+#include "Containers/UnrealString.h"
+
+namespace UE::DataflowVolume::Private
+{
+	/* --------------------------------------------------------------------------------------------------------------- */
+	/* FDataflowVolumeImpl */
+	/* --------------------------------------------------------------------------------------------------------------- */
+	
+	class FDataflowVolumeImpl
+	{
+	private:
+		openvdb::GridBase::Ptr BaseGrid;
+
+	public:
+		bool IsGridEmpty() const;
+
+		FVector GetVoxelSize() const;
+
+		FString GetGridName() const;
+
+		int32 GetActiveVoxelCount() const;
+
+		virtual void GetActiveVoxels(TArray<FBox>& OutActiveVoxels) const = 0;
+
+		FIntVector3 GetActiveVoxelDim() const;
+
+		FBox GetActiveVoxelBoundingBox() const;
+
+		FTransform GetGridTransform() const;
+
+		virtual FString VolumeInfo() const;
+
+		virtual ~FDataflowVolumeImpl() {}
+
+	protected:
+		void SetGrid(const openvdb::GridBase::Ptr& InGrid, const FString& InGridName);
+	};
+}
\ No newline at end of file
diff --git a/Engine/Plugins/Experimental/DataflowVolumetricPlugin/Source/DataflowVolumeCore/Private/Dataflow/DataflowVolumeUtils.cpp b/Engine/Plugins/Experimental/DataflowVolumetricPlugin/Source/DataflowVolumeCore/Private/Dataflow/DataflowVolumeUtils.cpp
new file mode 100644
--- /dev/null
+++ b/Engine/Plugins/Experimental/DataflowVolumetricPlugin/Source/DataflowVolumeCore/Private/Dataflow/DataflowVolumeUtils.cpp
@@ -0,0 +1,203 @@
+// Copyright Epic Games, Inc. All Rights Reserved.
+
+#include "Dataflow/DataflowVolumeUtils.h"
+#include "Dataflow/OpenVDB.h"
+#include "StaticMeshOperations.h"
+#include "StaticMeshAttributes.h"
+#include "RenderMath.h"
+
+namespace UE::DataflowVolumeUtils
+{
+	FMeshDescriptionAdapter::FMeshDescriptionAdapter(const FMeshDescription& InRawMesh, const openvdb::math::Transform& InTransform) :
+		RawMesh(&InRawMesh), Transform(InTransform)
+	{
+		InitializeCacheData();
+	}
+
+	FMeshDescriptionAdapter::FMeshDescriptionAdapter(const FMeshDescriptionAdapter& other)
+		: RawMesh(other.RawMesh), Transform(other.Transform)
+	{
+		InitializeCacheData();
+	}
+
+	void FMeshDescriptionAdapter::InitializeCacheData()
+	{
+		VertexPositions = RawMesh->GetVertexPositions();
+		TriangleCount = RawMesh->Triangles().Num();
+
+		IndexBuffer.Reserve(TriangleCount * 3);
+
+		for (const FTriangleID TriangleID : RawMesh->Triangles().GetElementIDs())
+		{
+			IndexBuffer.Add(RawMesh->GetTriangleVertexInstance(TriangleID, 0));
+			IndexBuffer.Add(RawMesh->GetTriangleVertexInstance(TriangleID, 1));
+			IndexBuffer.Add(RawMesh->GetTriangleVertexInstance(TriangleID, 2));
+		}
+
+
+	}
+
+	size_t FMeshDescriptionAdapter::polygonCount() const
+	{
+		return size_t(TriangleCount);
+	}
+
+	size_t FMeshDescriptionAdapter::pointCount() const
+	{
+		return size_t(RawMesh->Vertices().Num());
+	}
+
+	void FMeshDescriptionAdapter::getIndexSpacePoint(size_t FaceNumber, size_t CornerNumber, openvdb::Vec3d& pos) const
+	{
+		// Get the vertex position in local space.
+		const FVertexInstanceID VertexInstanceID = IndexBuffer[FaceNumber * 3 + CornerNumber];
+		// float3 position 
+		const FVertexID VertexID = RawMesh->GetVertexInstanceVertex(VertexInstanceID);
+
+
+		FVector3f Position = VertexPositions[VertexID];
+		pos = Transform.worldToIndex(openvdb::Vec3d(Position.X, Position.Y, Position.Z));
+	};
+
+	static void ResetStaticMeshDescription(FMeshDescription& MeshDecription)
+	{
+		MeshDecription.Empty();
+		FStaticMeshAttributes Attributes(MeshDecription);
+		Attributes.Register();
+	}
+
+	void AOSMeshToRawMesh(const FAOSMesh& AOSMesh, FMeshDescription& OutRawMesh)
+	{
+
+		ResetStaticMeshDescription(OutRawMesh);
+
+		FStaticMeshAttributes Attributes(OutRawMesh);
+		TVertexAttributesRef<FVector3f> VertexPositions = Attributes.GetVertexPositions();
+		TEdgeAttributesRef<bool> EdgeHardnesses = Attributes.GetEdgeHardnesses();
+		TPolygonGroupAttributesRef<FName> PolygonGroupImportedMaterialSlotNames = Attributes.GetPolygonGroupMaterialSlotNames();
+		TVertexInstanceAttributesRef<FVector3f> VertexInstanceNormals = Attributes.GetVertexInstanceNormals();
+		TVertexInstanceAttributesRef<FVector3f> VertexInstanceTangents = Attributes.GetVertexInstanceTangents();
+		TVertexInstanceAttributesRef<float> VertexInstanceBinormalSigns = Attributes.GetVertexInstanceBinormalSigns();
+		TVertexInstanceAttributesRef<FVector4f> VertexInstanceColors = Attributes.GetVertexInstanceColors();
+		TVertexInstanceAttributesRef<FVector2f> VertexInstanceUVs = Attributes.GetVertexInstanceUVs();
+
+		const uint32 DstNumPositions = AOSMesh.GetNumVertexes();
+		const uint32 DstNumIndexes = AOSMesh.GetNumIndexes();
+
+		if (VertexInstanceUVs.GetNumChannels() < 1)
+		{
+			VertexInstanceUVs.SetNumChannels(1);
+		}
+
+		FPolygonGroupID PolygonGroupID = INDEX_NONE;
+		if (OutRawMesh.PolygonGroups().Num() == 0)
+		{
+			PolygonGroupID = OutRawMesh.CreatePolygonGroup();
+			PolygonGroupImportedMaterialSlotNames[PolygonGroupID] = FName(*FString::Printf(TEXT("ProxyLOD_Material_%d"), FMath::Rand()));
+		}
+		else
+		{
+			PolygonGroupID = OutRawMesh.PolygonGroups().GetFirstValidID();
+		}
+
+		checkSlow(DstNumIndexes % 3 == 0);
+		// Copy the vertices over
+		TMap<int32, FVertexID> VertexIDMap;
+		VertexIDMap.Reserve(DstNumPositions);
+		{
+			const auto& AOSVertexes = AOSMesh.Vertexes;
+			for (uint32 i = 0, I = DstNumPositions; i < I; ++i)
+			{
+				const FVector3f& Position = AOSVertexes[i].GetPos();
+				const FVertexID NewVertexID = OutRawMesh.CreateVertex();
+				VertexPositions[NewVertexID] = Position;
+				VertexIDMap.Add(i, NewVertexID);
+			}
+
+			checkSlow(VertexPositions.GetNumElements() == DstNumPositions);
+		}
+
+		const uint32* AOSIndexes = AOSMesh.Indexes;
+
+		// Connectivity: 
+		auto CreateTriangle = [&OutRawMesh, PolygonGroupID, &VertexInstanceNormals, &VertexInstanceTangents, &VertexInstanceBinormalSigns, &VertexInstanceColors, &VertexInstanceUVs, &EdgeHardnesses](const FVertexID TriangleIndex[3], const FVector3f Normals[3])
+			{
+				TArray<FVertexInstanceID> VertexInstanceIDs;
+				VertexInstanceIDs.SetNum(3);
+				for (int32 Corner = 0; Corner < 3; ++Corner)
+				{
+					VertexInstanceIDs[Corner] = OutRawMesh.CreateVertexInstance(TriangleIndex[Corner]);
+					VertexInstanceTangents[VertexInstanceIDs[Corner]] = FVector3f(1, 0, 0);
+					VertexInstanceNormals[VertexInstanceIDs[Corner]] = Normals[Corner];
+					VertexInstanceBinormalSigns[VertexInstanceIDs[Corner]] = GetBasisDeterminantSign((FVector)VertexInstanceTangents[VertexInstanceIDs[Corner]].GetSafeNormal(),
+						(FVector)(VertexInstanceNormals[VertexInstanceIDs[Corner]] ^ VertexInstanceTangents[VertexInstanceIDs[Corner]]).GetSafeNormal(),
+						(FVector)VertexInstanceNormals[VertexInstanceIDs[Corner]].GetSafeNormal());
+					VertexInstanceColors[VertexInstanceIDs[Corner]] = FVector4f(1.0f);
+					VertexInstanceUVs.Set(VertexInstanceIDs[Corner], 0, FVector2f(0.0f, 0.0f));
+				}
+
+				// Insert a polygon into the mesh
+				OutRawMesh.CreatePolygon(PolygonGroupID, VertexInstanceIDs);
+			};
+
+		{
+			uint32 IndexStop = DstNumIndexes / 3;
+			for (uint32 t = 0, T = IndexStop; t < T; ++t)
+			{
+				FVertexID VertexIndexes[3];
+				FVector3f Normals[3];
+				for (int32 Corner = 0; Corner < 3; ++Corner)
+				{
+					VertexIndexes[Corner] = VertexIDMap[AOSIndexes[(t * 3) + Corner]];
+					const auto& AOSVertex = AOSMesh.Vertexes[AOSIndexes[(t * 3) + Corner]];
+					Normals[Corner] = AOSVertex.Normal;
+				}
+				CreateTriangle(VertexIndexes, Normals);
+			}
+		}
+	}
+
+	void ConvertMesh(const FAOSMesh& InMesh, FMeshDescription& OutMesh)
+	{
+		AOSMeshToRawMesh(InMesh, OutMesh);
+	}
+
+	/* ---------------------------------------------------------------------------------------------------- */
+
+	FMatrix CreateMatrixFromTransformAndVoxelSize(
+		const FVector InTranslation,
+		const FVector InRotation,
+		const FVector InScale,
+		const float InVoxelSize)
+	{
+		const FQuat Rotation = FQuat(FRotator(InRotation.Y, InRotation.Z, InRotation.X));
+		FTransform MeshXForm(Rotation, InTranslation, InScale);
+
+		FMatrix LocalToVoxel = FMatrix::Identity;
+		LocalToVoxel.M[0][0] = InVoxelSize;
+		LocalToVoxel.M[1][1] = InVoxelSize;
+		LocalToVoxel.M[2][2] = InVoxelSize;
+
+		FMatrix TransformMatrix = MeshXForm.ToMatrixWithScale().Inverse();
+		TransformMatrix = LocalToVoxel * TransformMatrix;
+
+		return TransformMatrix;
+	}
+
+	openvdb::Mat4R ConvertMatrixToVDBMatrix(FMatrix& InMatrix)
+	{
+		double* data = &InMatrix.M[0][0];
+		openvdb::Mat4R VDBMatDouble(data);
+
+		// NB: rounding errors in the inverse may have resulted in error in this col.
+		// openvdb explicitly checks this matrix row to insure the transform is affine and will throw 
+		VDBMatDouble.setCol(3, openvdb::Vec4R(0, 0, 0, 1));
+
+		return VDBMatDouble;
+	}
+
+
+
+
+}
+
diff --git a/Engine/Plugins/Experimental/DataflowVolumetricPlugin/Source/DataflowVolumeCore/Private/Dataflow/DataflowVolumeUtils.h b/Engine/Plugins/Experimental/DataflowVolumetricPlugin/Source/DataflowVolumeCore/Private/Dataflow/DataflowVolumeUtils.h
new file mode 100644
--- /dev/null
+++ b/Engine/Plugins/Experimental/DataflowVolumetricPlugin/Source/DataflowVolumeCore/Private/Dataflow/DataflowVolumeUtils.h
@@ -0,0 +1,700 @@
+// Copyright Epic Games, Inc. All Rights Reserved.
+
+#pragma once 
+
+#include "Dataflow/OpenVDB.h"
+#include "Math/Vector.h"
+#include "MeshDescription.h"
+
+#if PLATFORM_WINDOWS
+#include "Windows/AllowWindowsPlatformTypes.h"
+#endif
+
+THIRD_PARTY_INCLUDES_START
+PRAGMA_DISABLE_DEPRECATION_WARNINGS
+#include <tbb/blocked_range.h>
+#include <tbb/blocked_range2d.h>
+#include <tbb/concurrent_vector.h>
+#include <tbb/parallel_for.h>
+#include <tbb/parallel_reduce.h>
+#include <tbb/task_group.h>
+#include <tbb/task_scheduler_observer.h>
+PRAGMA_ENABLE_DEPRECATION_WARNINGS
+THIRD_PARTY_INCLUDES_END
+
+#if PLATFORM_WINDOWS
+#include "Windows/HideWindowsPlatformTypes.h"
+#endif
+
+struct FMeshDescription;
+
+namespace UE::DataflowVolumeUtils
+{
+	class FPlacedMesh
+	{
+	public:
+		FPlacedMesh()
+			: Mesh(nullptr)
+		{}
+
+		explicit FPlacedMesh(const FMeshDescription* MeshIn, const FTransform& TransformIn = FTransform::Identity)
+			: Mesh(MeshIn), Transform(TransformIn)
+		{}
+
+		FPlacedMesh(const FPlacedMesh& other)
+			: Mesh(other.Mesh)
+			, Transform(other.Transform)
+		{}
+
+		FPlacedMesh& operator=(const FPlacedMesh& other)
+		{
+			Mesh = other.Mesh;
+			Transform = other.Transform;
+			return *this;
+		}
+
+		const FMeshDescription* Mesh;
+		FTransform       Transform;
+	};
+
+	// the interrupter interface. 
+	struct FInterrupter
+	{
+		FInterrupter() = default;
+		virtual ~FInterrupter() = default;
+		// templated openvdb code requires these functions.
+		virtual void start(const char* name = nullptr) { (void)name; }
+		virtual void end() {};
+		virtual bool wasInterrupted(int percent = -1) { (void)percent; return false; }
+	};
+
+	/**
+	*  Mesh adapter class that implements the needed API for the methods in openvdb that convert
+	*  polygonal mesh to a signed distance field.
+	*/
+	class FMeshDescriptionAdapter
+	{
+	public:
+		/**
+		* Constructor
+		*/
+		FMeshDescriptionAdapter(const FMeshDescription& InRawMesh, const openvdb::math::Transform& InTransform);
+		FMeshDescriptionAdapter(const FMeshDescriptionAdapter& other);
+
+		// Total number of polygons
+		size_t polygonCount() const;
+
+		// Total number of points (vertex locations)
+		size_t pointCount() const;
+
+		// Vertex count for polygon n: currently FMeshDescription is just triangles.
+		size_t vertexCount(size_t n) const { return 3; }
+
+		// Return position pos in local grid index space for polygon n and vertex v
+		void getIndexSpacePoint(size_t FaceNumber, size_t CornerNumber, openvdb::Vec3d& pos) const;
+
+		/**
+		* The transform used to map between the physical space of the mesh and the voxel space.
+		*/
+		const openvdb::math::Transform& GetTransform() const
+		{
+			return Transform;
+		}
+
+	private:
+
+		// Pointer to the raw mesh we are wrapping
+		const FMeshDescription* RawMesh;
+
+
+
+		//////////////////////////////////////////////////////////////////////////
+		//Cache data
+		void InitializeCacheData();
+		uint32 TriangleCount;
+		TVertexAttributesConstRef<FVector3f> VertexPositions;
+		// Local version of the index array.  The FMeshDescription doesn't really have one.
+		TArray<FVertexInstanceID> IndexBuffer;
+
+
+
+		// Transform used to convert the mesh space into the index space used by voxelization
+		const openvdb::math::Transform Transform;
+	};
+
+
+
+	/**
+	* Mesh type that holds minimal required data for openvdb iso-surface extraction
+	* interface. Implements the openvdb MeshDataAdapter interface.
+	*
+	* NB: std::vector required by the openvdb interface.
+	*/
+	struct FMixedPolyMesh
+	{
+		std::vector<openvdb::Vec3s> Points;
+		std::vector<openvdb::Vec4I> Quads;
+		std::vector<openvdb::Vec3I> Triangles;
+		openvdb::math::Transform	Transform;
+
+		// ~MeshDataAdapter Interface Begin
+		// https://www.openvdb.org/documentation/doxygen/interfaceMeshDataAdapter.html
+
+		// Total number of polygons
+		size_t polygonCount() const
+		{
+			return Quads.size() + Triangles.size();
+		}
+
+		// Total number of points (vertex locations)
+		size_t pointCount() const
+		{
+			return Points.size();
+		}
+
+		// Vertex count for polygon n
+		size_t vertexCount(size_t n) const
+		{
+			return PolygonIsQuad(n) ? 4 : 3;
+		}
+
+		// Return position pos in local grid index space for polygon n and vertex v
+		void getIndexSpacePoint(size_t n, size_t v, openvdb::Vec3d& pos) const
+		{
+			if (PolygonIsQuad(n))
+			{
+				pos = Points[Quads[GetQuadIndex(n)][v]];
+			}
+			else
+			{
+				pos = Points[Triangles[GetTriangleIndex(n)][v]];
+			}
+
+			pos = Transform.worldToIndex(pos);
+		}
+		// ~MeshDataAdapter Interface End
+
+	private:
+		bool PolygonIsQuad(size_t n) const
+		{
+			return n < Quads.size();
+		}
+
+		bool PolygonIsTriangle(size_t n) const
+		{
+			return n >= Triangles.size();
+		}
+
+		size_t GetQuadIndex(size_t n) const
+		{
+			return n;
+		}
+
+		size_t GetTriangleIndex(size_t n) const
+		{
+			return n - Quads.size();
+		}
+	};
+
+
+
+	/**
+	* Simplifier vertex type that stores position and normal.
+	*
+	* Implements the interface needed by templated simplifier code.
+	*/
+	class FPositionNormalVertex
+	{
+		typedef FPositionNormalVertex  VertType;
+	public:
+		FPositionNormalVertex() {}
+
+		enum { NumAttributesValue = 3 };
+
+		uint32          MaterialIndex = 0;
+		FVector3f			Position;
+		// ---- Attributes --------
+		// currently just the components of the normal.
+		FVector3f			Normal;
+
+
+		// Attributes are all assumed to be floats.  
+		static uint32 NumAttributes()
+		{
+
+			return (sizeof(VertType) - sizeof(uint32) - sizeof(FVector3f)) / sizeof(float);
+		}
+
+		// Methods needed by the simplifier.
+		uint32			GetMaterialIndex() const { return MaterialIndex; }
+		FVector3f& GetPos() { return Position; }
+		const FVector3f& GetPos() const { return Position; }
+		float* GetAttributes() { return (float*)&Normal; }
+		const float* GetAttributes() const { return (const float*)&Normal; }
+
+
+		void Correct() { Normal.Normalize(); }
+
+		// Note this uses exact float compares..
+		bool operator==(const VertType& OtherVert) const
+		{
+			bool Result = true;
+			if (MaterialIndex != OtherVert.MaterialIndex ||
+				Position != OtherVert.Position ||
+				Normal != OtherVert.Normal)
+			{
+				Result = false;
+			}
+
+			return Result;
+		}
+		bool operator!=(const VertType& OtherVert)
+		{
+			return !(*this == OtherVert);
+		}
+
+		// NB: Addition isn't commutative since the MaterialIndex..
+		// maybe Result.MaterialIndex = min(this->MaterialIndex, Other.MaterialIndex) 
+		// would be better.
+		VertType operator+(const VertType& OtherVert) const
+		{
+			VertType Result(*this);
+
+			Result.Position += OtherVert.Position;
+			Result.Normal += OtherVert.Normal;
+
+			return Result;
+		}
+
+		// NB: A-B != -(B-A) because of the material index.
+		VertType operator-(const VertType& OtherVert) const
+		{
+			VertType Result(*this);
+
+			Result.Position -= OtherVert.Position;
+			Result.Normal -= OtherVert.Normal;
+
+			return Result;
+		}
+
+		// NB: The normal isn't unit length after this..
+		VertType operator*(const float Scalar) const
+		{
+			VertType Result(*this);
+			Result.Position *= Scalar;
+			Result.Normal *= Scalar;
+
+			return Result;
+		}
+
+
+		VertType operator/(const float Scalar) const
+		{
+			float ia = 1.0f / Scalar;
+			return (*this) * ia;
+		}
+	};
+
+
+
+	/**
+	* A Triangle Mesh Array Of Structs (AOS) for the vertex data.  This should work with the simplification code.
+	*/
+	template <typename SimplifierVertexType>
+	class TAOSMesh
+	{
+	public:
+
+		typedef SimplifierVertexType  VertexType;
+
+		TAOSMesh();
+		TAOSMesh(int32 VertCount, int32 FaceCount);
+
+		~TAOSMesh();
+
+		// reset the mesh to size zero, deleting content.
+		void Empty();
+
+		// resize the mesh, deleting the current content.
+		void Resize(int32 VertCount, int32 FaceCount);
+
+		// set the real vertex & index count after duplicate removal
+		void SetVertexAndIndexCount(uint32 VertCount, uint32 IndexCount);
+
+		// Swap content with an existing mesh of the same type.
+		void Swap(TAOSMesh& other);
+
+
+		// Get the indices that correspond to this face.
+		openvdb::Vec3I GetFace(int32 FaceNumber) const;
+
+		// Get an array of positions only.
+		void GetPosArray(std::vector<FVector3f>& PosArray) const;
+
+
+		// This method is for testing only.  Not designed for perf.
+		int32 ComputeDegenerateTriCount() const;
+
+
+		bool IsEmpty() const { return (NumVertexes == 0 || NumIndexes == 0); }
+		uint32 GetNumVertexes() const { return NumVertexes; }
+		uint32 GetNumIndexes()  const { return NumIndexes; }
+
+		// unfettered public access.
+		VertexType* Vertexes;
+		uint32* Indexes;
+
+	private:
+
+		uint32 NumVertexes;
+		uint32 NumIndexes;
+
+		TAOSMesh(const TAOSMesh& other);
+	};
+
+
+
+	/**
+	* Array of structs mesh composed of vertices that are compatible with the simplifier.
+	*/
+	typedef TAOSMesh<FPositionNormalVertex>  FAOSMesh;
+
+	/* Threaded wrappers */
+	typedef tbb::blocked_range<int32>    FIntRange;
+	typedef tbb::blocked_range<uint32>   FUIntRange;
+	typedef tbb::blocked_range2d<int32>  FIntRange2d;
+
+	template <typename RangeType, typename FunctorType>
+	void Parallel_For(const RangeType& Range, const FunctorType& Functor, const bool bParallel = true)
+	{
+		TRACE_CPUPROFILER_EVENT_SCOPE(ProxyLOD::Parallel_For)
+			if (bParallel) // run in parallel
+			{
+				// Functor can be passed by reference since we wait until completion
+				tbb::parallel_for(Range,
+					[&Functor](const RangeType& Range)
+					{
+						// #TODO Investigate why Insights stops working when used
+						//TRACE_CPUPROFILER_EVENT_SCOPE(ProxyLOD::Parallel_For)
+						Functor(Range);
+					});
+			}
+			else // single threaded
+			{
+				Functor(Range);
+			}
+	}
+
+	// --- Implementation of templated code. ---
+
+
+	template <typename SimplifierVertexType>
+	TAOSMesh<SimplifierVertexType>::TAOSMesh() :
+		Vertexes(NULL),
+		Indexes(NULL),
+		NumVertexes(0),
+		NumIndexes(0)
+	{}
+
+	template <typename SimplifierVertexType>
+	TAOSMesh<SimplifierVertexType>::TAOSMesh(int32 VertCount, int32 FaceCount) :
+		Vertexes(NULL),
+		Indexes(NULL),
+		NumVertexes(0),
+		NumIndexes(0)
+	{
+		this->Resize(VertCount, FaceCount);
+	}
+
+	template <typename SimplifierVertexType>
+	TAOSMesh<SimplifierVertexType>::~TAOSMesh()
+	{
+		this->Empty();
+	}
+
+	template <typename SimplifierVertexType>
+	void TAOSMesh<SimplifierVertexType>::Empty()
+	{
+		if (Vertexes) delete[] Vertexes;
+		if (Indexes)  delete[] Indexes;
+
+		Vertexes = NULL;
+		Indexes = NULL;
+		NumVertexes = 0;
+		NumIndexes = 0;
+	}
+
+	template <typename SimplifierVertexType>
+	void TAOSMesh<SimplifierVertexType>::Resize(int32 VertCount, int32 FaceCount)
+	{
+		check(VertCount > -1); check(FaceCount > -1);
+
+		this->Empty();
+
+		if (FaceCount > 0 && VertCount > 0)
+		{
+			NumVertexes = VertCount;
+			NumIndexes = FaceCount * 3;
+
+			Vertexes = new VertexType[VertCount];
+			Indexes = new uint32[NumIndexes];
+		}
+
+	}
+
+	template <typename SimplifierVertexType>
+	void TAOSMesh<SimplifierVertexType>::SetVertexAndIndexCount(uint32 VertCount, uint32 IndexCount)
+	{
+		check(VertCount <= NumVertexes);
+		check(IndexCount <= NumIndexes);
+
+		NumVertexes = VertCount;
+		NumIndexes = IndexCount;
+	}
+
+	template <typename SimplifierVertexType>
+	void TAOSMesh<SimplifierVertexType>::Swap(TAOSMesh& other)
+	{
+		// Swap sizes
+		{
+			uint32 Tmp = other.NumVertexes;
+			other.NumVertexes = NumVertexes;
+			NumVertexes = Tmp;
+
+			Tmp = other.NumIndexes;
+			other.NumIndexes = NumIndexes;
+			NumIndexes = Tmp;
+		}
+
+		// Swap Vertex Pointer
+		{
+			VertexType* TmpVp = other.Vertexes;
+			other.Vertexes = Vertexes;
+			Vertexes = TmpVp;
+		}
+
+		// Sway index pointer
+		{
+			uint32* TmpIdp = other.Indexes;
+			other.Indexes = Indexes;
+			Indexes = TmpIdp;
+		}
+
+	}
+
+	template <typename SimplifierVertexType>
+	openvdb::Vec3I TAOSMesh<SimplifierVertexType>::GetFace(int32 FaceNumber) const
+	{
+		check(FaceNumber > -1); check(uint32(FaceNumber * 3) < NumIndexes);
+
+		uint32 offset = FaceNumber * 3;
+		return openvdb::Vec3I(Indexes[offset], Indexes[offset + 1], Indexes[offset + 2]);
+	}
+
+	template <typename SimplifierVertexType>
+	void TAOSMesh<SimplifierVertexType>::GetPosArray(std::vector<FVector3f>& PosArray) const
+	{
+		// resize the target array
+		PosArray.resize(NumVertexes);
+
+		// copy the data over.
+		Parallel_For(FUIntRange(0, NumVertexes), [this, &PosArray](const FUIntRange& Range)
+			{
+				FVector3f* Pos = PosArray.data();
+				FPositionNormalVertex* VertexArray = this->Vertexes;
+
+
+				for (uint32 r = Range.begin(), R = Range.end(); r < R; ++r)
+				{
+					Pos[r] = VertexArray[r].GetPos();
+				}
+			});
+	}
+
+	template <typename SimplifierVertexType>
+	int32 TAOSMesh<SimplifierVertexType>::ComputeDegenerateTriCount() const
+	{
+		int32 DegenerateTriCount = 0;
+		uint32 NumTris = GetNumIndexes() / 3;
+		for (uint32 i = 0; i < NumTris; ++i)
+		{
+			uint32 Offset = i * 3;
+			uint32 Idx[3] = { Indexes[Offset], Indexes[Offset + 1], Indexes[Offset + 2] };
+			FVector Tri[3] = { Vertexes[Idx[0]].GetPos(), Vertexes[Idx[1]].GetPos(), Vertexes[Idx[2]].GetPos() };
+
+			const FVector NormalDir = (Tri[2] - Tri[0]) ^ (Tri[1] - Tri[0]);
+
+			if (NormalDir.SizeSquared() == 0.0f)
+			{
+				DegenerateTriCount++;
+			}
+		}
+
+		return DegenerateTriCount;
+	}
+
+
+
+	/**
+	* Convert a mixed triangle and quad mesh to a triangle mesh type by splitting quads.
+	* No new vertices are created.
+	*
+	* NB: This does not attempt to add any additional attributes to the result.  Will
+	*     need to separately compute vertex normals if desired.
+	*
+	* @param InMesh   Mesh with a mixture of quads and triangles.
+	* @param OutMesh  Triangle Mesh.
+	* @param bClockWise how to order the verts in a triangle
+	*/
+	template <typename T>
+	void MixedPolyMeshToAOSMesh(const FMixedPolyMesh& InMesh, TAOSMesh<T>& OutMesh, const bool bClockWise = true);
+
+	// Convert MixedPolyMesh to AOS Mesh.  This requires splitting quads to produce triangles.
+	template <typename T>
+	void MixedPolyMeshToAOSMesh(const FMixedPolyMesh& MixedPolyMesh, TAOSMesh<T>& DstAOSMesh, bool bClockWise)
+	{
+
+		// Splitting a quad doesn't introduce any new verts.
+		const uint32 DstNumVerts = MixedPolyMesh.Points.size();
+
+		const uint32 NumQuads = MixedPolyMesh.Quads.size();
+
+		// Each quad becomes 2 triangles.
+		const uint32 DstNumTris = 2 * NumQuads + MixedPolyMesh.Triangles.size();
+
+		// Each Triangle has 3 corners
+		const uint32 DstNumIndexes = 3 * DstNumTris;
+
+		// Empty and Allocate space
+
+		DstAOSMesh.Resize(DstNumVerts, DstNumTris);
+
+		// Copy the vertices position over and give it a dummy material index.
+		{
+			// Allocate the space for the verts in the DstAOSMesh
+
+			Parallel_For(FUIntRange(0, DstNumVerts),
+				[&MixedPolyMesh, &DstAOSMesh](const FUIntRange& Range)
+				{
+					for (uint32 i = Range.begin(), I = Range.end(); i < I; ++i)
+					{
+						const openvdb::Vec3s& Point = MixedPolyMesh.Points[i];
+						DstAOSMesh.Vertexes[i].Position = FVector3f(Point[0], Point[1], Point[2]);
+						DstAOSMesh.Vertexes[i].MaterialIndex = 0;
+					}
+				});
+		}
+
+		// Split VDB Quads
+		{
+
+
+			// NB: The Quads are ordered in clockwise fashion.
+
+			Parallel_For(FUIntRange(0, NumQuads),
+				[&MixedPolyMesh, &DstAOSMesh, bClockWise](const FUIntRange& Range)
+				{
+					uint32* Indices = DstAOSMesh.Indexes;
+					for (uint32 q = Range.begin(), Q = Range.end(); q < Q; ++q)
+					{
+						const uint32 Offset = q * 6;
+						const openvdb::Vec4I& Quad = MixedPolyMesh.Quads[q];
+						// add as two triangles
+						if (bClockWise)
+						{
+							// first triangle
+							Indices[Offset] = Quad[0];
+							Indices[Offset + 1] = Quad[1];
+							Indices[Offset + 2] = Quad[2];
+							// second triangle
+							Indices[Offset + 3] = Quad[2];
+							Indices[Offset + 4] = Quad[3];
+							Indices[Offset + 5] = Quad[0];
+						}
+						else
+						{
+							// first triangle
+							Indices[Offset] = Quad[0];
+							Indices[Offset + 1] = Quad[3];
+							Indices[Offset + 2] = Quad[2];
+
+							// second triangle
+							Indices[Offset + 3] = Quad[2];
+							Indices[Offset + 4] = Quad[1];
+							Indices[Offset + 5] = Quad[0];
+						}
+					}
+
+				});
+
+			// add the MixedPolyMesh triangles.
+			Parallel_For(FUIntRange(0, MixedPolyMesh.Triangles.size()),
+				[&MixedPolyMesh, &DstAOSMesh, NumQuads, bClockWise](const FUIntRange& Range)
+				{
+					uint32* Indices = DstAOSMesh.Indexes;
+					for (uint32 t = Range.begin(), EndT = Range.end(); t < EndT; ++t)
+					{
+						const uint32 Offset = NumQuads * 6 + t * 3;
+						const openvdb::Vec3I& Tri = MixedPolyMesh.Triangles[t];
+						// add the triangle
+						if (bClockWise)
+						{
+							Indices[Offset] = Tri[0];
+							Indices[Offset + 1] = Tri[1];
+							Indices[Offset + 2] = Tri[2];
+						}
+						else
+						{
+							Indices[Offset] = Tri[2];
+							Indices[Offset + 1] = Tri[1];
+							Indices[Offset + 2] = Tri[0];
+						}
+					}
+				});
+		}
+	}
+
+
+
+	/**
+	* Common interface to convert between various triangle-based mesh types.
+	*
+	* The conversions will maintain geometry and connectivity, but different mesh types support
+	* different attribute types and possibly different frequenceies even when the same attribute types
+	* exist on both mesh types ( e.g. FMeshDescription has per-wedge data while FVertexDataMesh has per-vertex attributes).
+	*
+	* @param InMesh   Source Mesh to convert.
+	* @param OutMesh  The result of the mesh conversion.  Any data already stored in OutMesh will be lost.
+	*
+	*/
+	void ConvertMesh(const FAOSMesh& InMesh, FMeshDescription& OutMesh);
+
+
+
+
+	/**
+	* Convert the simplifier-friendly array-of-structs mesh to a struct of arrays FMeshDescription.
+	*
+	* NB: This copies the AOSVertex normal to the wedge FMeshDescription::Tangentz.
+	*     But additional FMeshDescription attributes, including tangent/bitangent are given default values.
+	*
+	* @param InMesh   Source Mesh to convert.
+	* @param OutMesh  The result of the mesh conversion.  Any data already stored in OutMesh will be lost.
+	*
+	*/
+	void AOSMeshToRawMesh(const FAOSMesh& InMesh, FMeshDescription& OutMesh);
+
+	/**
+	* Create a world-to-index-space transform matrix from transform and voxelSize
+	*/
+	FMatrix CreateMatrixFromTransformAndVoxelSize(
+		const FVector InTranslation, 
+		const FVector InRotation,
+		const FVector InScale,
+		const float InVoxelSize);
+
+	/**
+	* Convert a FMatrix to openvdb::Mat4R
+	*/
+	openvdb::Mat4R ConvertMatrixToVDBMatrix(FMatrix& InMatrix);
+}
diff --git a/Engine/Plugins/Experimental/DataflowVolumetricPlugin/Source/DataflowVolumeCore/Private/Dataflow/OpenVDB.h b/Engine/Plugins/Experimental/DataflowVolumetricPlugin/Source/DataflowVolumeCore/Private/Dataflow/OpenVDB.h
new file mode 100644
--- /dev/null
+++ b/Engine/Plugins/Experimental/DataflowVolumetricPlugin/Source/DataflowVolumeCore/Private/Dataflow/OpenVDB.h
@@ -0,0 +1,78 @@
+// Copyright Epic Games, Inc. All Rights Reserved.
+
+#pragma once
+
+#include "HAL/Platform.h"
+#include "Misc/CoreMiscDefines.h"
+
+#if PLATFORM_WINDOWS
+#include "Windows/AllowWindowsPlatformTypes.h"
+#endif
+
+THIRD_PARTY_INCLUDES_START
+
+UE_PUSH_MACRO("check")
+#undef check
+
+#ifdef _MSC_VER
+	#pragma warning (push)
+	#pragma warning (disable : 4146) // unary minus operator applied to unsigned type, result still unsigned
+	#pragma warning (disable : 4706) // assignment within conditional expression
+#endif
+
+#ifndef M_PI
+    #define M_PI    3.14159265358979323846
+	#define LOCAL_M_PI 1
+#endif
+
+#ifndef M_PI_2
+    #define M_PI_2  1.57079632679489661923 // pi/2
+	#define LOCAL_M_PI_2 1
+#endif
+
+#define BOOST_ALLOW_DEPRECATED_HEADERS
+
+#include <openvdb/openvdb.h>
+#include <openvdb/tools/Composite.h> 		// for csgUnion
+#include <openvdb/tools/Interpolation.h> 	// for grid sampler
+#include <openvdb/tools/MeshToVolume.h> 	// for MeshToVolume
+#include <openvdb/tools/VolumeToMesh.h> 	// for VolumeToMesh
+#include <openvdb/tools/LevelSetSphere.h>   // Generate a narrow-band level set of sphere
+#include <openvdb/tools/LevelSetPlatonic.h> // Generate a narrow-band level sets of the five platonic solids
+#include <openvdb/points/PointCount.h>
+#include <openvdb/tools/VolumeToSpheres.h>
+#include <openvdb/tools/Morphology.h>
+#include <openvdb/math/Math.h> // for isFinite()
+#include <openvdb/tools/ChangeBackground.h>
+#include <openvdb/tools/GridTransformer.h> // for resampleToMatch()
+#include <openvdb/tools/LevelSetRebuild.h> // for levelSetRebuild()
+#include <openvdb/tools/Prune.h>
+#include <openvdb/tools/SignedFloodFill.h>
+#include <openvdb/util/NullInterrupter.h>
+#include <openvdb/Types.h>
+#include <openvdb/points/PointScatter.h>
+#include <openvdb/math/Math.h>
+#include <openvdb/tools/LevelSetUtil.h> 
+#include <openvdb/tools/LevelSetFracture.h>
+#include <openvdb/tools/ParticlesToLevelSet.h>
+#include <openvdb/tools/GridOperators.h>
+
+#if LOCAL_M_PI
+	#undef M_PI
+#endif
+
+#if LOCAL_M_PI_2
+	#undef M_PI_2
+#endif
+
+#ifdef _MSC_VER
+	#pragma warning (pop)
+#endif
+
+UE_POP_MACRO("check")
+
+THIRD_PARTY_INCLUDES_END
+
+#if PLATFORM_WINDOWS
+#include "Windows/HideWindowsPlatformTypes.h"
+#endif
\ No newline at end of file
diff --git a/Engine/Plugins/Experimental/DataflowVolumetricPlugin/Source/DataflowVolumeNodes/DataflowVolumeNodes.Build.cs b/Engine/Plugins/Experimental/DataflowVolumetricPlugin/Source/DataflowVolumeNodes/DataflowVolumeNodes.Build.cs
new file mode 100644
--- /dev/null
+++ b/Engine/Plugins/Experimental/DataflowVolumetricPlugin/Source/DataflowVolumeNodes/DataflowVolumeNodes.Build.cs
@@ -0,0 +1,42 @@
+// Copyright Epic Games, Inc. All Rights Reserved.
+
+namespace UnrealBuildTool.Rules
+{
+	public class DataflowVolumeNodes : ModuleRules
+	{
+        public DataflowVolumeNodes (ReadOnlyTargetRules Target) : base(Target)
+		{
+			PublicDependencyModuleNames.AddRange(
+				new string[]
+				{
+				}
+			);
+
+			PrivateDependencyModuleNames.AddRange(
+				new string[] 
+				{
+					"DynamicMesh",
+					"MeshConversion",
+					"Core",
+					"CoreUObject",
+					"ChaosCore",
+					"Chaos",
+					"DataflowCore",
+					"DataflowEngine",
+					"DataflowEnginePlugin",
+					"DataflowNodes",
+					"GeometryCollectionEngine",
+					"GeometryCore",
+					"Voronoi",
+					"GeometryFramework",
+					"MeshDescription",
+					"StaticMeshDescription",
+					"PlanarCut",
+					"Engine",
+					"DataflowEditor",
+					"DataflowVolumeCore"
+				}
+			);
+		}
+	}
+}
diff --git a/Engine/Plugins/Experimental/DataflowVolumetricPlugin/Source/DataflowVolumeNodes/Private/Dataflow/DataflowVolumeNodeAndPinTypeColors.cpp b/Engine/Plugins/Experimental/DataflowVolumetricPlugin/Source/DataflowVolumeNodes/Private/Dataflow/DataflowVolumeNodeAndPinTypeColors.cpp
new file mode 100644
--- /dev/null
+++ b/Engine/Plugins/Experimental/DataflowVolumetricPlugin/Source/DataflowVolumeNodes/Private/Dataflow/DataflowVolumeNodeAndPinTypeColors.cpp
@@ -0,0 +1,23 @@
+// Copyright Epic Games, Inc. All Rights Reserved.
+
+#include "Dataflow/DataflowVolumeNodeAndPinTypeColors.h"
+#include "Dataflow/DataflowNode.h"
+#include "Math/Color.h" 
+#include "Dataflow/DataflowNodeColorsRegistry.h"
+
+namespace UE::Dataflow
+{
+	void RegisterVolumeColors()
+	{
+		static const FLinearColor CDefaultNodeBodyTintColor = FLinearColor(0.f, 0.f, 0.f, 0.5f);
+		static const float CDefaultWireThickness = 1.5f;
+
+		// Node colors
+		DATAFLOW_NODE_REGISTER_CREATION_FACTORY_NODE_COLORS_BY_CATEGORY("Volume", FLinearColor(0.273f, 0.172f, 1.f), CDefaultNodeBodyTintColor);
+
+		// PinType colors
+		DATAFLOW_NODE_REGISTER_CREATION_FACTORY_PIN_SETTINGS_BY_TYPE("FDataflowFloatVolume", FLinearColor(.263f, 0.412f, 0.9f, 1.0f), CDefaultWireThickness);
+		DATAFLOW_NODE_REGISTER_CREATION_FACTORY_PIN_SETTINGS_BY_TYPE("FDataflowIntVolume", FLinearColor(.363f, 0.512f, 0.9f, 1.0f), CDefaultWireThickness);
+		DATAFLOW_NODE_REGISTER_CREATION_FACTORY_PIN_SETTINGS_BY_TYPE("FDataflowVectorVolume", FLinearColor(.463f, 0.612f, 0.9f, 1.0f), CDefaultWireThickness);
+	}
+}
diff --git a/Engine/Plugins/Experimental/DataflowVolumetricPlugin/Source/DataflowVolumeNodes/Private/Dataflow/DataflowVolumeNodeAndPinTypeColors.h b/Engine/Plugins/Experimental/DataflowVolumetricPlugin/Source/DataflowVolumeNodes/Private/Dataflow/DataflowVolumeNodeAndPinTypeColors.h
new file mode 100644
--- /dev/null
+++ b/Engine/Plugins/Experimental/DataflowVolumetricPlugin/Source/DataflowVolumeNodes/Private/Dataflow/DataflowVolumeNodeAndPinTypeColors.h
@@ -0,0 +1,9 @@
+// Copyright Epic Games, Inc. All Rights Reserved.
+
+#pragma once 
+
+namespace UE::Dataflow
+{
+	void RegisterVolumeColors();
+}
+
diff --git a/Engine/Plugins/Experimental/DataflowVolumetricPlugin/Source/DataflowVolumeNodes/Private/Dataflow/DataflowVolumeNodes.cpp b/Engine/Plugins/Experimental/DataflowVolumetricPlugin/Source/DataflowVolumeNodes/Private/Dataflow/DataflowVolumeNodes.cpp
new file mode 100644
--- /dev/null
+++ b/Engine/Plugins/Experimental/DataflowVolumetricPlugin/Source/DataflowVolumeNodes/Private/Dataflow/DataflowVolumeNodes.cpp
@@ -0,0 +1,1986 @@
+// Copyright Epic Games, Inc. All Rights Reserved.
+
+#include "Dataflow/DataflowVolumeNodes.h"
+#include "Dataflow/DataflowCore.h"
+//#include "UDynamicMesh.h"
+#include "StaticMeshAttributes.h"
+#include "PlanarCut.h"
+//#include "TransformTypes.h"
+#include "GeometryCollection/GeometryCollectionEngineConversion.h"
+#include "Dataflow/DataflowOverlay.h"
+#include "Dataflow/DataflowDebugDrawInterface.h"
+#if WITH_EDITOR
+#include "Dataflow/DataflowRenderingViewMode.h"
+#endif
+#include "GeometryCollection/Facades/PointsCollectionFacade.h"
+
+#include UE_INLINE_GENERATED_CPP_BY_NAME(DataflowVolumeNodes)
+
+namespace UE::Dataflow
+{
+	void RegisterVolumeNodes()
+	{
+		// Generators
+		DATAFLOW_NODE_REGISTER_CREATION_FACTORY(FMakeSDFSphereDataflowNode);
+		DATAFLOW_NODE_REGISTER_CREATION_FACTORY(FMakeSDFCubeDataflowNode);
+		DATAFLOW_NODE_REGISTER_CREATION_FACTORY(FMakeSDFPlatonicSolidDataflowNode);
+		// Utilities
+		DATAFLOW_NODE_REGISTER_CREATION_FACTORY(FVolumeInfoDataflowNode);
+		DATAFLOW_NODE_REGISTER_CREATION_FACTORY(FVolumeAnalysisDataflowNode);
+		DATAFLOW_NODE_REGISTER_CREATION_FACTORY(FVolumeCombineDataflowNode);
+		// Convert
+		DATAFLOW_NODE_REGISTER_CREATION_FACTORY(FCollectionToVolumeDataflowNode);
+		DATAFLOW_NODE_REGISTER_CREATION_FACTORY(FVolumeToSpheresDataflowNode);
+		DATAFLOW_NODE_REGISTER_CREATION_FACTORY(FConvertVolumeDataflowNode);
+
+		// Scatter
+		DATAFLOW_NODE_REGISTER_CREATION_FACTORY(FUniformVolumeScatterDataflowNode);
+		DATAFLOW_NODE_REGISTER_CREATION_FACTORY(FDenseUniformVolumeScatterDataflowNode);
+		DATAFLOW_NODE_REGISTER_CREATION_FACTORY(FNonUniformVolumeScatterDataflowNode);
+
+		// Sample
+		DATAFLOW_NODE_REGISTER_CREATION_FACTORY(FVolumeSampleToPointsDataflowNode);
+
+		// Visualize
+		DATAFLOW_NODE_REGISTER_CREATION_FACTORY(FVolumeSliceDataflowNode);
+	}
+}
+
+/* --------------------------------------------------------------------------------------------- */
+
+FMakeSDFSphereDataflowNode::FMakeSDFSphereDataflowNode(const UE::Dataflow::FNodeParameters& InParam, FGuid InGuid)
+	: FDataflowNode(InParam, InGuid)
+{
+	RegisterInputConnection(&VoxelSize);
+	RegisterInputConnection(&Radius).SetCanHidePin(true).SetPinIsHidden(true);
+	RegisterInputConnection(&Center).SetCanHidePin(true).SetPinIsHidden(true);
+	RegisterOutputConnection(&FloatVolume);
+
+	DebugDrawRenderSettings.RenderType = EDataflowDebugDrawRenderType::Wireframe;
+	DebugDrawRenderSettings.bTranslucent = false;
+	DebugDrawRenderSettings.Color = FLinearColor::Green;
+	DebugDrawRenderSettings.LineWidthMultiplier = 1.5f;
+}
+
+void FMakeSDFSphereDataflowNode::Evaluate(UE::Dataflow::FContext& Context, const FDataflowOutput* Out) const
+{
+	if (Out->IsA(&FloatVolume))
+	{
+		const float InVoxelSize = GetValue(Context, &VoxelSize);
+		const float InRadius = GetValue(Context, &Radius);
+		const FVector InCenter = GetValue(Context, &Center);
+
+		FDataflowFloatVolume OutFloatVolume = FDataflowFloatVolume::CreateSphereSDF(InVoxelSize, InRadius, InCenter);
+
+		SetValue(Context, OutFloatVolume, &FloatVolume);
+	}
+}
+
+#if WITH_EDITOR
+bool FMakeSDFSphereDataflowNode::CanDebugDrawViewMode(const FName& ViewModeName) const
+{
+	return ViewModeName == UE::Dataflow::FDataflowConstruction3DViewMode::Name;
+}
+
+void FMakeSDFSphereDataflowNode::DebugDraw(UE::Dataflow::FContext& Context, IDataflowDebugDrawInterface& DataflowRenderingInterface, const FDebugDrawParameters& DebugDrawParameters) const
+{
+	using namespace UE::Dataflow::Overlay;
+
+	if (DebugDrawParameters.bNodeIsSelected || DebugDrawParameters.bNodeIsPinned)
+	{
+		FString OverlayEvaluateResultString;
+
+		const FDataflowOutput* FloatVolumeOutput = FindOutput(&FloatVolume);
+		if (FloatVolumeOutput)
+		{
+			const FDataflowFloatVolume& OutFloatVolume = FloatVolumeOutput->GetValue(Context, FloatVolume);
+//			const FDataflowFloatVolume& OutFloatVolume = this->GetOutputValue(Context, FloatVolume);
+
+			DebugDrawRenderSettings.SetDebugDrawSettings(DataflowRenderingInterface);
+
+			TArray<FBox> ActiveVoxels;
+			OutFloatVolume.GetActiveVoxels(ActiveVoxels);
+			
+			for (const FBox& Voxel : ActiveVoxels)
+			{
+				if (bDrawVoxelCenters)
+				{
+					DataflowRenderingInterface.DrawPoint(0.5 * FVector(Voxel.Max + Voxel.Min));
+				}
+
+				if (bDrawVoxels)
+				{
+					DataflowRenderingInterface.DrawBox(0.5 * FVector(Voxel.Max - Voxel.Min), FQuat::Identity, 0.5 * FVector(Voxel.Min + Voxel.Max), 1.0);
+				}
+			}
+
+			OverlayEvaluateResultString = OutFloatVolume.VolumeInfo();
+		}
+
+		// ------------------------------------------------------------------------------------------------------
+		// Overlay
+		TArray<FString> OverlayStrArr;
+		OverlayStrArr.Add(BuildOverlayNodeInfoString(this));
+		OverlayStrArr.Add(OverlayEvaluateResultString);
+		FString OverlayStr = BuildOverlayFinalString(OverlayStrArr);
+		DataflowRenderingInterface.DrawOverlayText(OverlayStr);
+		// ------------------------------------------------------------------------------------------------------
+	}
+}
+#endif
+
+/* --------------------------------------------------------------------------------------------- */
+
+FMakeSDFCubeDataflowNode::FMakeSDFCubeDataflowNode(const UE::Dataflow::FNodeParameters& InParam, FGuid InGuid)
+	: FDataflowNode(InParam, InGuid)
+{
+	RegisterInputConnection(&Scale).SetCanHidePin(true).SetPinIsHidden(true);
+	RegisterInputConnection(&Center).SetCanHidePin(true).SetPinIsHidden(true);
+	RegisterInputConnection(&VoxelSize).SetCanHidePin(true).SetPinIsHidden(true);
+	RegisterOutputConnection(&FloatVolume);
+
+	DebugDrawRenderSettings.RenderType = EDataflowDebugDrawRenderType::Wireframe;
+	DebugDrawRenderSettings.bTranslucent = false;
+	DebugDrawRenderSettings.Color = FLinearColor::Green;
+	DebugDrawRenderSettings.LineWidthMultiplier = 1.5f;
+}
+
+void FMakeSDFCubeDataflowNode::Evaluate(UE::Dataflow::FContext& Context, const FDataflowOutput* Out) const
+{
+	if (Out->IsA(&FloatVolume))
+	{
+		const float InScale = GetValue(Context, &Scale);
+		const FVector InCenter = GetValue(Context, &Center);
+		const float InVoxelSize = GetValue(Context, &VoxelSize);
+
+		FDataflowFloatVolume OutFloatVolume = FDataflowFloatVolume::CreateCubeSDF(
+			InScale,
+			InCenter,
+			InVoxelSize);
+
+		SetValue(Context, OutFloatVolume, &FloatVolume);
+	}
+}
+
+#if WITH_EDITOR
+bool FMakeSDFCubeDataflowNode::CanDebugDrawViewMode(const FName& ViewModeName) const
+{
+	return ViewModeName == UE::Dataflow::FDataflowConstruction3DViewMode::Name;
+}
+
+void FMakeSDFCubeDataflowNode::DebugDraw(UE::Dataflow::FContext& Context, IDataflowDebugDrawInterface& DataflowRenderingInterface, const FDebugDrawParameters& DebugDrawParameters) const
+{
+	using namespace UE::Dataflow::Overlay;
+
+	if (DebugDrawParameters.bNodeIsSelected || DebugDrawParameters.bNodeIsPinned)
+	{
+		FString OverlayEvaluateResultString;
+
+		const FDataflowOutput* FloatVolumeOutput = FindOutput(&FloatVolume);
+		if (FloatVolumeOutput)
+		{
+			const FDataflowFloatVolume& OutFloatVolume = FloatVolumeOutput->GetValue(Context, FloatVolume);
+			//			const FDataflowFloatVolume& OutFloatVolume = this->GetOutputValue(Context, FloatVolume);
+
+			DebugDrawRenderSettings.SetDebugDrawSettings(DataflowRenderingInterface);
+
+			TArray<FBox> ActiveVoxels;
+			OutFloatVolume.GetActiveVoxels(ActiveVoxels);
+
+			for (const FBox& Voxel : ActiveVoxels)
+			{
+				if (bDrawVoxelCenters)
+				{
+					DataflowRenderingInterface.DrawPoint(0.5 * FVector(Voxel.Max + Voxel.Min));
+				}
+
+				if (bDrawVoxels)
+				{
+					DataflowRenderingInterface.DrawBox(0.5 * FVector(Voxel.Max - Voxel.Min), FQuat::Identity, 0.5 * FVector(Voxel.Min + Voxel.Max), 1.0);
+				}
+			}
+
+			OverlayEvaluateResultString = OutFloatVolume.VolumeInfo();
+		}
+
+		// ------------------------------------------------------------------------------------------------------
+		// Overlay
+		TArray<FString> OverlayStrArr;
+		OverlayStrArr.Add(BuildOverlayNodeInfoString(this));
+		OverlayStrArr.Add(OverlayEvaluateResultString);
+		FString OverlayStr = BuildOverlayFinalString(OverlayStrArr);
+		DataflowRenderingInterface.DrawOverlayText(OverlayStr);
+		// ------------------------------------------------------------------------------------------------------
+	}
+}
+#endif
+
+/* --------------------------------------------------------------------------------------------- */
+
+FMakeSDFPlatonicSolidDataflowNode::FMakeSDFPlatonicSolidDataflowNode(const UE::Dataflow::FNodeParameters& InParam, FGuid InGuid)
+	: FDataflowNode(InParam, InGuid)
+{
+	RegisterInputConnection(&Scale).SetCanHidePin(true).SetPinIsHidden(true);
+	RegisterInputConnection(&Center).SetCanHidePin(true).SetPinIsHidden(true);
+	RegisterInputConnection(&VoxelSize).SetCanHidePin(true).SetPinIsHidden(true);
+	RegisterOutputConnection(&FloatVolume);
+
+	DebugDrawRenderSettings.RenderType = EDataflowDebugDrawRenderType::Wireframe;
+	DebugDrawRenderSettings.bTranslucent = false;
+	DebugDrawRenderSettings.Color = FLinearColor::Green;
+	DebugDrawRenderSettings.LineWidthMultiplier = 1.5f;
+}
+
+void FMakeSDFPlatonicSolidDataflowNode::Evaluate(UE::Dataflow::FContext& Context, const FDataflowOutput* Out) const
+{
+	if (Out->IsA(&FloatVolume))
+	{
+		int32 FaceCount = 0;
+		switch (PlatonicSolidType)
+		{
+			case EDataflowVolumePlatonicSolidType::Tetrahedron:
+				FaceCount = 4;
+				break;
+			case EDataflowVolumePlatonicSolidType::Cube:
+				FaceCount = 6;
+				break;
+			case EDataflowVolumePlatonicSolidType::Octahedron:
+				FaceCount = 8;
+				break;
+			case EDataflowVolumePlatonicSolidType::Dodecahedron:
+				FaceCount = 12;
+				break;
+			case EDataflowVolumePlatonicSolidType::Icosahedron:
+				FaceCount = 20;
+				break;
+			default:
+				break;
+		}
+
+		const float InScale = GetValue(Context, &Scale);
+		const FVector InCenter = GetValue(Context, &Center);
+		const float InVoxelSize = GetValue(Context, &VoxelSize);
+
+		FDataflowFloatVolume OutFloatVolume = FDataflowFloatVolume::CreatePlatonicSolidSDF(
+			FaceCount,
+			InScale,
+			InCenter,
+			InVoxelSize);
+
+		SetValue(Context, OutFloatVolume, &FloatVolume);
+	}
+}
+
+#if WITH_EDITOR
+bool FMakeSDFPlatonicSolidDataflowNode::CanDebugDrawViewMode(const FName& ViewModeName) const
+{
+	return ViewModeName == UE::Dataflow::FDataflowConstruction3DViewMode::Name;
+}
+
+void FMakeSDFPlatonicSolidDataflowNode::DebugDraw(UE::Dataflow::FContext& Context, IDataflowDebugDrawInterface& DataflowRenderingInterface, const FDebugDrawParameters& DebugDrawParameters) const
+{
+	using namespace UE::Dataflow::Overlay;
+
+	if (DebugDrawParameters.bNodeIsSelected || DebugDrawParameters.bNodeIsPinned)
+	{
+		FString OverlayEvaluateResultString;
+
+		const FDataflowOutput* FloatVolumeOutput = FindOutput(&FloatVolume);
+		if (FloatVolumeOutput)
+		{
+			const FDataflowFloatVolume& OutFloatVolume = FloatVolumeOutput->GetValue(Context, FloatVolume);
+			//			const FDataflowFloatVolume& OutFloatVolume = this->GetOutputValue(Context, FloatVolume);
+
+			DebugDrawRenderSettings.SetDebugDrawSettings(DataflowRenderingInterface);
+
+			TArray<FBox> ActiveVoxels;
+			OutFloatVolume.GetActiveVoxels(ActiveVoxels);
+
+			for (const FBox& Voxel : ActiveVoxels)
+			{
+				if (bDrawVoxelCenters)
+				{
+					DataflowRenderingInterface.DrawPoint(0.5 * FVector(Voxel.Max + Voxel.Min));
+				}
+
+				if (bDrawVoxels)
+				{
+					DataflowRenderingInterface.DrawBox(0.5 * FVector(Voxel.Max - Voxel.Min), FQuat::Identity, 0.5 * FVector(Voxel.Min + Voxel.Max), 1.0);
+				}
+			}
+
+			OverlayEvaluateResultString = OutFloatVolume.VolumeInfo();
+		}
+
+		// ------------------------------------------------------------------------------------------------------
+		// Overlay
+		TArray<FString> OverlayStrArr;
+		OverlayStrArr.Add(BuildOverlayNodeInfoString(this));
+		OverlayStrArr.Add(OverlayEvaluateResultString);
+		FString OverlayStr = BuildOverlayFinalString(OverlayStrArr);
+		DataflowRenderingInterface.DrawOverlayText(OverlayStr);
+		// ------------------------------------------------------------------------------------------------------
+	}
+}
+#endif
+
+/* --------------------------------------------------------------------------------------------- */
+
+FVolumeInfoDataflowNode::FVolumeInfoDataflowNode(const UE::Dataflow::FNodeParameters& InParam, FGuid InGuid)
+	: FDataflowNode(InParam, InGuid)
+{
+	RegisterInputConnection(&FloatVolume);
+	RegisterOutputConnection(&String);
+}
+
+void FVolumeInfoDataflowNode::Evaluate(UE::Dataflow::FContext& Context, const FDataflowOutput* Out) const
+{
+	if (Out->IsA(&String))
+	{
+		const FDataflowFloatVolume& InFloatVolume = GetValue(Context, &FloatVolume);
+
+		FString OutString = InFloatVolume.VolumeInfo();
+
+		SetValue(Context, MoveTemp(OutString), &String);
+	}
+}
+
+/* --------------------------------------------------------------------------------------------- */
+
+FCollectionToVolumeDataflowNode::FCollectionToVolumeDataflowNode(const UE::Dataflow::FNodeParameters& InParam, FGuid InGuid)
+	: FDataflowNode(InParam, InGuid)
+{
+	RegisterInputConnection(&Collection);
+	RegisterInputConnection(&ReferenceVolume).SetCanHidePin(true).SetPinIsHidden(true);
+	RegisterInputConnection(&VoxelSize).SetCanHidePin(true).SetPinIsHidden(true);
+	RegisterInputConnection(&IsoValue).SetCanHidePin(true).SetPinIsHidden(true);
+	RegisterOutputConnection(&FloatVolume);
+	RegisterOutputConnection(&FaceIndexVolume);
+
+	DebugDrawRenderSettings.RenderType = EDataflowDebugDrawRenderType::Shaded;
+	DebugDrawRenderSettings.bTranslucent = false;
+	DebugDrawRenderSettings.Color = FLinearColor::Green;
+	DebugDrawRenderSettings.LineWidthMultiplier = 1.5f;
+}
+
+void FCollectionToVolumeDataflowNode::Evaluate(UE::Dataflow::FContext& Context, const FDataflowOutput* Out) const
+{
+	if (Out->IsA(&FloatVolume) || Out->IsA(&FaceIndexVolume))
+	{
+		if (IsConnected(&Collection))
+		{
+			const FManagedArrayCollection& InCollection = GetValue(Context, &Collection);
+			const float InVoxelSize = GetValue(Context, &VoxelSize);
+
+			// Anytype returns storage type
+			FDataflowVolume InReferenceVolume = GetValue(Context, &ReferenceVolume);
+
+			if (InCollection.NumElements(FGeometryCollection::TransformGroup) > 0)
+			{
+				if (TUniquePtr<FGeometryCollection> GeomCollection = TUniquePtr<FGeometryCollection>(InCollection.NewCopy<FGeometryCollection>()))
+				{
+					const TManagedArray<FTransform3f>& BoneTransforms = InCollection.GetAttribute<FTransform3f>("Transform", FGeometryCollection::TransformGroup);
+
+					TArray<int32> TransformIndices;
+					TransformIndices.AddUninitialized(BoneTransforms.Num());
+
+					int32 Idx = 0;
+					for (int32& TransformIdx : TransformIndices)
+					{
+						TransformIdx = Idx++;
+					}
+
+					FMeshDescription MeshDescription;
+					FStaticMeshAttributes Attributes(MeshDescription);
+					Attributes.Register();
+
+					FTransform TransformOut;
+
+					ConvertToMeshDescription(MeshDescription, TransformOut, /*bCenterPivot*/true, *GeomCollection, BoneTransforms, TransformIndices);
+
+					FDataflowIntVolume OutFaceIndexVolume;
+
+					FDataflowFloatVolume OutVolume = FDataflowFloatVolume::CreateVolumeFromMeshDescription(
+						MeshDescription,
+						TransformOut,
+						InVoxelSize,
+						InReferenceVolume,
+						OutputType,
+						GridName,
+						bUseWorldSpaceUnits,
+						ExteriorBand,
+						InteriorBand,
+						ExteriorBandVoxels,
+						InteriorBandVoxels,
+						bFillInterior,
+						bPreserveHoles,
+						IsoValue,
+						OutFaceIndexVolume);
+
+					SetValue(Context, OutVolume, &FloatVolume);
+					SetValue(Context, OutFaceIndexVolume, &FaceIndexVolume);
+
+					return;
+				}
+			}
+		}
+
+		SetValue(Context, FDataflowFloatVolume(), &FloatVolume);
+		SetValue(Context, FDataflowIntVolume(), &FaceIndexVolume);
+	}
+}
+
+#if WITH_EDITOR
+bool FCollectionToVolumeDataflowNode::CanDebugDrawViewMode(const FName& ViewModeName) const
+{
+	return ViewModeName == UE::Dataflow::FDataflowConstruction3DViewMode::Name;
+}
+
+void FCollectionToVolumeDataflowNode::DebugDraw(UE::Dataflow::FContext& Context, IDataflowDebugDrawInterface& DataflowRenderingInterface, const FDebugDrawParameters& DebugDrawParameters) const
+{
+	using namespace UE::Dataflow::Overlay;
+
+	if (DebugDrawParameters.bNodeIsSelected || DebugDrawParameters.bNodeIsPinned)
+	{
+		FString OverlayEvaluateResultString;
+
+		const FDataflowOutput* FloatVolumeOutput = FindOutput(&FloatVolume);
+		if (FloatVolumeOutput)
+		{
+			const FDataflowFloatVolume& OutFloatVolume = FloatVolumeOutput->GetValue(Context, FloatVolume);
+			//			const FDataflowFloatVolume& OutFloatVolume = this->GetOutputValue(Context, FloatVolume);
+
+			DebugDrawRenderSettings.SetDebugDrawSettings(DataflowRenderingInterface);
+
+			TArray<FBox> ActiveVoxels;
+			OutFloatVolume.GetActiveVoxels(ActiveVoxels);
+
+			for (const FBox& Voxel : ActiveVoxels)
+			{
+				if (bDrawVoxelCenters)
+				{
+					DataflowRenderingInterface.DrawPoint(0.5 * FVector(Voxel.Max + Voxel.Min));
+				}
+
+				if (bDrawVoxels)
+				{
+					DataflowRenderingInterface.DrawBox(0.5 * FVector(Voxel.Max - Voxel.Min), FQuat::Identity, 0.5 * FVector(Voxel.Min + Voxel.Max), 1.0);
+				}
+			}
+
+			OverlayEvaluateResultString = OutFloatVolume.VolumeInfo();
+		}
+
+		// ------------------------------------------------------------------------------------------------------
+		// Overlay
+		TArray<FString> OverlayStrArr;
+		OverlayStrArr.Add(BuildOverlayNodeInfoString(this));
+		OverlayStrArr.Add(OverlayEvaluateResultString);
+		FString OverlayStr = BuildOverlayFinalString(OverlayStrArr);
+		DataflowRenderingInterface.DrawOverlayText(OverlayStr);
+		// ------------------------------------------------------------------------------------------------------
+	}
+}
+#endif
+
+/* --------------------------------------------------------------------------------------------- */
+
+FVolumeToSpheresDataflowNode::FVolumeToSpheresDataflowNode(const UE::Dataflow::FNodeParameters& InParam, FGuid InGuid)
+	: FDataflowNode(InParam, InGuid)
+{
+	RegisterInputConnection(&FloatVolume);
+	RegisterInputConnection(&MinSphereCount).SetCanHidePin(true).SetPinIsHidden(true);
+	RegisterInputConnection(&MaxSphereCount).SetCanHidePin(true).SetPinIsHidden(true);
+	RegisterInputConnection(&MinRadius).SetCanHidePin(true).SetPinIsHidden(true);
+	RegisterInputConnection(&MaxRadius).SetCanHidePin(true).SetPinIsHidden(true);
+	RegisterInputConnection(&IsoValue).SetCanHidePin(true).SetPinIsHidden(true);
+	RegisterInputConnection(&InstanceCount).SetCanHidePin(true).SetPinIsHidden(true);
+	RegisterOutputConnection(&Collection);
+	RegisterOutputConnection(&Spheres);
+
+	DebugDrawRenderSettings.RenderType = EDataflowDebugDrawRenderType::Wireframe;
+	DebugDrawRenderSettings.bTranslucent = false;
+	DebugDrawRenderSettings.Color = FLinearColor::Blue;
+	DebugDrawRenderSettings.LineWidthMultiplier = 1.5f;
+}
+
+void FVolumeToSpheresDataflowNode::Evaluate(UE::Dataflow::FContext& Context, const FDataflowOutput* Out) const
+{
+	if (Out->IsA(&Collection) || Out->IsA(&Spheres))
+	{
+		if (IsConnected(&FloatVolume))
+		{
+			const FDataflowFloatVolume& InFloatVolume = GetValue(Context, &FloatVolume);
+			const int32 InMinSphereCount = GetValue(Context, &MinSphereCount);
+			const int32 InMaxSphereCount = GetValue(Context, &MaxSphereCount);
+			const float InMinRadius = GetValue(Context, &MinRadius);
+			const float InMaxRadius = GetValue(Context, &MaxRadius);
+			const float InIsoValue = GetValue(Context, &IsoValue);
+			const int32 InInstanceCount = GetValue(Context, &InstanceCount);
+
+			TArray<FVector> OutSphereCenters;
+			TArray<float> OutSphereRadii;
+
+			// Compute sphere packing
+			InFloatVolume.VolumeToSpheres(
+				InMinSphereCount,
+				InMaxSphereCount,
+				bOverlapping,
+				InMinRadius,
+				InMaxRadius,
+				InIsoValue,
+				InInstanceCount,
+				OutSphereCenters,
+				OutSphereRadii);
+
+			// Add spheres to OutCollection
+			FManagedArrayCollection OutCollection;
+
+			const int32 NumSpheresFromSDF = OutSphereCenters.Num();
+
+			if (NumSpheresFromSDF > 0)
+			{
+				// Add new element to groups
+				static const FName SpheresGroup = "Spheres";
+
+				const int32 NumTransforms = OutCollection.NumElements(FGeometryCollection::TransformGroup);
+				const int32 NumSpheres = OutCollection.NumElements(SpheresGroup);
+
+				OutCollection.AddElements(1, FGeometryCollection::TransformGroup);
+				OutCollection.AddElements(NumSpheresFromSDF, SpheresGroup);
+
+				TManagedArray<FTransform>& Transform = OutCollection.AddAttribute<FTransform>("Transform", FGeometryCollection::TransformGroup);
+				TManagedArray<FString>& BoneName = OutCollection.AddAttribute<FString>("BoneName", FGeometryCollection::TransformGroup);
+				TManagedArray<FLinearColor>& BoneColor = OutCollection.AddAttribute<FLinearColor>("BoneColor", FGeometryCollection::TransformGroup);
+				TManagedArray<int32>& Parent = OutCollection.AddAttribute<int32>("Parent", FGeometryCollection::TransformGroup);
+				TManagedArray<TSet<int32>>& Children = OutCollection.AddAttribute<TSet<int32>>("Children", FGeometryCollection::TransformGroup);
+				TManagedArray<int32>& TransformToSphereIndex = OutCollection.AddAttribute<int32>("TransformToSphereIndex", FGeometryCollection::TransformGroup);
+				TManagedArray<int32>& BoneMap = OutCollection.AddAttribute<int32>("BoneMap", SpheresGroup);
+				TManagedArray<int32>& TransformIndex = OutCollection.AddAttribute<int32>("TransformIndex", SpheresGroup);
+				TManagedArray<FVector3f>& Center = OutCollection.AddAttribute<FVector3f>("Center", SpheresGroup);
+				TManagedArray<float>& Radius = OutCollection.AddAttribute<float>("Radius", SpheresGroup);
+
+				Transform[NumTransforms] = FTransform::Identity;
+				BoneName[NumTransforms] = FString(TEXT("Spheres"));
+				BoneColor[NumTransforms] = FLinearColor(0.02f, 0.01f, 0.1f, 1.0f);
+				Parent[NumTransforms] = -1;
+				Children[NumTransforms] = TSet<int32>();
+				TransformToSphereIndex[NumTransforms] = NumSpheres;
+
+				TArray<FSphere> OutSpheres; OutSpheres.AddUninitialized(NumSpheresFromSDF);
+
+				for (int32 Idx = 0; Idx < NumSpheresFromSDF; ++Idx)
+				{
+					Center[NumSpheres + Idx] = (FVector3f)OutSphereCenters[Idx];
+					Radius[NumSpheres + Idx] = OutSphereRadii[Idx];
+
+					OutSpheres[Idx].Center = OutSphereCenters[Idx];
+					OutSpheres[Idx].W = OutSphereRadii[Idx];
+				}
+
+				TransformIndex[NumSpheres] = NumTransforms;
+
+				SetValue(Context, MoveTemp(OutCollection), &Collection);
+				SetValue(Context, MoveTemp(OutSpheres), &Spheres);
+
+				return;
+			}
+		}
+
+		SetValue(Context, FManagedArrayCollection(), &Collection);
+		SetValue(Context, TArray<FSphere>(), &Spheres);
+	}
+}
+
+#if WITH_EDITOR
+bool FVolumeToSpheresDataflowNode::CanDebugDrawViewMode(const FName& ViewModeName) const
+{
+	return ViewModeName == UE::Dataflow::FDataflowConstruction3DViewMode::Name;
+}
+
+void FVolumeToSpheresDataflowNode::DebugDraw(UE::Dataflow::FContext& Context, IDataflowDebugDrawInterface& DataflowRenderingInterface, const FDebugDrawParameters& DebugDrawParameters) const
+{
+	using namespace UE::Dataflow::Overlay;
+
+	if (DebugDrawParameters.bNodeIsSelected || DebugDrawParameters.bNodeIsPinned)
+	{
+		FString OverlayEvaluateResultString;
+
+		const FDataflowOutput* SpheresOutput = FindOutput(&Spheres);
+		if (SpheresOutput)
+		{
+			const TArray<FSphere>& OutSpheres = SpheresOutput->GetValue(Context, Spheres);
+			//			const FDataflowFloatVolume& OutFloatVolume = this->GetOutputValue(Context, FloatVolume);
+
+			DebugDrawRenderSettings.SetDebugDrawSettings(DataflowRenderingInterface);
+
+			float MinRad = FLT_MAX, MaxRad = -FLT_MAX, SumRad = 0;
+			for (const FSphere& Sphere : OutSpheres)
+			{
+				if (Sphere.W < MinRad)
+				{
+					MinRad = Sphere.W;
+				}
+				if (Sphere.W > MaxRad)
+				{
+					MaxRad = Sphere.W;
+				}
+
+				SumRad += Sphere.W;
+
+				if (bDrawSphereCenters)
+				{
+					DataflowRenderingInterface.DrawPoint(Sphere.Center);
+				}
+
+				if (bDrawSpheres)
+				{
+					DataflowRenderingInterface.DrawSphere(Sphere.Center, Sphere.W);
+				}
+			}
+
+			const int32 NumSpheres = OutSpheres.Num();
+			const float MeanRad = SumRad / (float)NumSpheres;
+
+			OverlayEvaluateResultString.Appendf(TEXT("Number of Spheres: %d\nMin radius: %3.3f\nMax radius: %3.3f\nMedian radius: %3.3f\n"),
+				NumSpheres,
+				MinRad,
+				MaxRad,
+				MeanRad);
+		}
+
+		// ------------------------------------------------------------------------------------------------------
+		// Overlay
+		TArray<FString> OverlayStrArr;
+		OverlayStrArr.Add(BuildOverlayNodeInfoString(this));
+		OverlayStrArr.Add(OverlayEvaluateResultString);
+		FString OverlayStr = BuildOverlayFinalString(OverlayStrArr);
+		DataflowRenderingInterface.DrawOverlayText(OverlayStr);
+		// ------------------------------------------------------------------------------------------------------
+	}
+}
+#endif
+
+/* --------------------------------------------------------------------------------------------- */
+
+FConvertVolumeDataflowNode::FConvertVolumeDataflowNode(const UE::Dataflow::FNodeParameters& InParam, FGuid InGuid)
+	: FDataflowNode(InParam, InGuid)
+{
+	RegisterInputConnection(&FloatVolume);
+	RegisterOutputConnection(&Collection);
+	RegisterOutputConnection(&FloatVolume, &FloatVolume);
+
+	DebugDrawRenderSettings.RenderType = EDataflowDebugDrawRenderType::Shaded;
+	DebugDrawRenderSettings.bTranslucent = false;
+	DebugDrawRenderSettings.Color = FLinearColor::Green;
+	DebugDrawRenderSettings.LineWidthMultiplier = 1.5f;
+}
+
+void FConvertVolumeDataflowNode::Evaluate(UE::Dataflow::FContext& Context, const FDataflowOutput* Out) const
+{
+	if (Out->IsA(&FloatVolume))
+	{
+		if (IsConnected(&FloatVolume))
+		{
+			if (ConvertTo == EDataflowVolumeConvertSDFTo::Volume)
+			{
+				const FDataflowFloatVolume& InFloatVolume = GetValue(Context, &FloatVolume);
+
+				if (!(GridClass == EDataflowVolumeConvertSDFGridClass::NoChange &&
+					GridType == EDataflowVolumeConvertSDFGridType::NoChange))
+				{
+					if (GridClass == EDataflowVolumeConvertSDFGridClass::SDFToFog)
+					{
+						FDataflowFloatVolume OutVolume = InFloatVolume.CovertSDFToFogVolume(
+							bPruneTolerance,
+							Tolerance,
+							bFloodFillOutput,
+							bActivateInteriror);
+
+						SetValue(Context, MoveTemp(OutVolume), &FloatVolume);
+					}
+					else if (GridClass == EDataflowVolumeConvertSDFGridClass::FogToSDF)
+					{
+						FDataflowFloatVolume OutVolume = InFloatVolume.CovertFogVolumeToSDF(
+							FogIsoValue,
+							bPruneTolerance,
+							Tolerance,
+							bFloodFillOutput,
+							bActivateInteriror);
+
+						SetValue(Context, MoveTemp(OutVolume), &FloatVolume);
+					}
+				}
+				else
+				{
+					SafeForwardInput(Context, &FloatVolume, &FloatVolume);
+				}
+
+				SetValue(Context, FManagedArrayCollection(), &Collection);
+			}
+			else if (ConvertTo == EDataflowVolumeConvertSDFTo::Collection)
+			{
+				SetValue(Context, FDataflowFloatVolume(), &FloatVolume);
+			}
+
+			return;
+		}
+
+		SetValue(Context, FDataflowFloatVolume(), &FloatVolume);
+	}
+	else if (Out->IsA(&Collection))
+	{
+		if (IsConnected(&FloatVolume))
+		{
+			if (ConvertTo == EDataflowVolumeConvertSDFTo::Collection)
+			{
+
+				const FDataflowFloatVolume& InFloatVolume = GetValue(Context, &FloatVolume);
+
+				FMeshDescription OutMeshDescription;
+
+				InFloatVolume.ConvertVolumeToMeshDescription(
+					IsoValue,
+					Adaptivity,
+					OutMeshDescription);
+
+				FManagedArrayCollection OutCollection = FManagedArrayCollection();
+
+				if (OutMeshDescription.Vertices().Num() > 0)
+				{
+					// Convert MeshDescription to FManagedArrayCollection
+					FGeometryCollection NewGeometryCollection = FGeometryCollection();
+					FGeometryCollectionEngineConversion::AppendMeshDescription(&OutMeshDescription, FString(TEXT("TEST")), 0, FTransform().Identity, &NewGeometryCollection);
+
+					NewGeometryCollection.CopyTo(&OutCollection);
+				}
+
+				SetValue(Context, MoveTemp(OutCollection), &Collection);
+				SafeForwardInput(Context, &FloatVolume, &FloatVolume);
+
+				return;
+			}
+			else if (ConvertTo == EDataflowVolumeConvertSDFTo::Volume)
+			{
+				SetValue(Context, FManagedArrayCollection(), &Collection);
+			}
+		}
+	}
+}
+
+#if WITH_EDITOR
+bool FConvertVolumeDataflowNode::CanDebugDrawViewMode(const FName& ViewModeName) const
+{
+	return ViewModeName == UE::Dataflow::FDataflowConstruction3DViewMode::Name;
+}
+
+void FConvertVolumeDataflowNode::DebugDraw(UE::Dataflow::FContext& Context, IDataflowDebugDrawInterface& DataflowRenderingInterface, const FDebugDrawParameters& DebugDrawParameters) const
+{
+	using namespace UE::Dataflow::Overlay;
+
+	if (DebugDrawParameters.bNodeIsSelected || DebugDrawParameters.bNodeIsPinned)
+	{
+		if (ConvertTo == EDataflowVolumeConvertSDFTo::Volume)
+		{
+			FString OverlayEvaluateResultString;
+
+			const FDataflowOutput* FloatVolumeOutput = FindOutput(&FloatVolume);
+			if (FloatVolumeOutput)
+			{
+				const FDataflowFloatVolume& OutFloatVolume = FloatVolumeOutput->GetValue(Context, FloatVolume);
+				//			const FDataflowFloatVolume& OutFloatVolume = this->GetOutputValue(Context, FloatVolume);
+
+				// Check if grid is empty
+//				if (!OutFloatVolume.GetFloatVolume()->IsGridEmpty())
+				if (0)
+				{
+					DebugDrawRenderSettings.SetDebugDrawSettings(DataflowRenderingInterface);
+
+					TArray<FBox> ActiveVoxels;
+					OutFloatVolume.GetActiveVoxels(ActiveVoxels);
+
+					for (const FBox& Voxel : ActiveVoxels)
+					{
+						if (bDrawVoxelCenters)
+						{
+							DataflowRenderingInterface.DrawPoint(0.5 * FVector(Voxel.Max + Voxel.Min));
+						}
+
+						if (bDrawVoxels)
+						{
+							DataflowRenderingInterface.DrawBox(0.5 * FVector(Voxel.Max - Voxel.Min), FQuat::Identity, 0.5 * FVector(Voxel.Min + Voxel.Max), 1.0);
+						}
+					}
+
+					OverlayEvaluateResultString = OutFloatVolume.VolumeInfo();
+
+					// ------------------------------------------------------------------------------------------------------
+					// Overlay
+					TArray<FString> OverlayStrArr;
+					OverlayStrArr.Add(BuildOverlayNodeInfoString(this));
+					OverlayStrArr.Add(OverlayEvaluateResultString);
+					FString OverlayStr = BuildOverlayFinalString(OverlayStrArr);
+					DataflowRenderingInterface.DrawOverlayText(OverlayStr);
+					// ------------------------------------------------------------------------------------------------------
+				}
+			}
+		}
+	}
+}
+#endif
+
+/* --------------------------------------------------------------------------------------------- */
+
+FUniformVolumeScatterDataflowNode::FUniformVolumeScatterDataflowNode(const UE::Dataflow::FNodeParameters& InParam, FGuid InGuid)
+	: FDataflowNode(InParam, InGuid)
+{
+	RegisterInputConnection(&FloatVolume);
+	RegisterInputConnection(&MinNumberOfPoints).SetCanHidePin(true).SetPinIsHidden(true);
+	RegisterInputConnection(&MaxNumberOfPoints).SetCanHidePin(true).SetPinIsHidden(true);
+	RegisterInputConnection(&RandomSeed).SetCanHidePin(true).SetPinIsHidden(true);
+	RegisterInputConnection(&IsoValue).SetCanHidePin(true).SetPinIsHidden(true);
+	RegisterInputConnection(&Spread).SetCanHidePin(true).SetPinIsHidden(true);
+	RegisterOutputConnection(&Points);
+}
+
+void FUniformVolumeScatterDataflowNode::Evaluate(UE::Dataflow::FContext& Context, const FDataflowOutput* Out) const
+{
+	if (Out->IsA(&Points))
+	{
+		if (IsConnected(&FloatVolume))
+		{
+			const FDataflowFloatVolume& InFloatVolume = GetValue(Context, &FloatVolume);
+			const int32 InMinNumberOfPoints = GetValue(Context, &MinNumberOfPoints);
+			const int32 InMaxNumberOfPoints = GetValue(Context, &MaxNumberOfPoints);
+			const int32 InRandomSeed = GetValue(Context, &RandomSeed);
+			const float InIsoValue = GetValue(Context, &IsoValue);
+			const float InSpread = GetValue(Context, &Spread);
+
+			TArray<FVector> OutPoints;
+
+			// Compute point scatter in volume
+			InFloatVolume.UniformVolumeScatter(
+				InMinNumberOfPoints,
+				InMaxNumberOfPoints,
+				InRandomSeed,
+				InIsoValue,
+				InSpread,
+				OutPoints);
+
+			SetValue(Context, MoveTemp(OutPoints), &Points);
+
+			return;
+		}
+
+		SetValue(Context, TArray<FVector>(), &Points);
+	}
+}
+
+#if WITH_EDITOR
+bool FUniformVolumeScatterDataflowNode::CanDebugDrawViewMode(const FName& ViewModeName) const
+{
+	return ViewModeName == UE::Dataflow::FDataflowConstruction3DViewMode::Name;
+}
+
+void FUniformVolumeScatterDataflowNode::DebugDraw(UE::Dataflow::FContext& Context, IDataflowDebugDrawInterface& DataflowRenderingInterface, const FDebugDrawParameters& DebugDrawParameters) const
+{
+	using namespace UE::Dataflow::Overlay;
+
+	if (DebugDrawParameters.bNodeIsSelected || DebugDrawParameters.bNodeIsPinned)
+	{
+		FString OverlayEvaluateResultString;
+
+		const FDataflowOutput* PointsOutput = FindOutput(&Points);
+		if (PointsOutput)
+		{
+			const TArray<FVector>& OutPoints = PointsOutput->GetValue(Context, Points);
+			//			const FDataflowFloatVolume& OutFloatVolume = this->GetOutputValue(Context, FloatVolume);
+
+			DataflowRenderingInterface.SetWorldPriority();
+			DataflowRenderingInterface.SetColor(Color);
+
+			for (const FVector& Point : OutPoints)
+			{
+				DataflowRenderingInterface.DrawPoint(Point);
+			}
+
+			OverlayEvaluateResultString.Appendf(TEXT("Number of Points: %d\n"), OutPoints.Num());
+		}
+
+		// ------------------------------------------------------------------------------------------------------
+		// Overlay
+		TArray<FString> OverlayStrArr;
+		OverlayStrArr.Add(BuildOverlayNodeInfoString(this));
+		OverlayStrArr.Add(OverlayEvaluateResultString);
+		FString OverlayStr = BuildOverlayFinalString(OverlayStrArr);
+		DataflowRenderingInterface.DrawOverlayText(OverlayStr);
+		// ------------------------------------------------------------------------------------------------------
+	}
+}
+#endif
+
+/* --------------------------------------------------------------------------------------------- */
+
+FDenseUniformVolumeScatterDataflowNode::FDenseUniformVolumeScatterDataflowNode(const UE::Dataflow::FNodeParameters& InParam, FGuid InGuid)
+	: FDataflowNode(InParam, InGuid)
+{
+	RegisterInputConnection(&FloatVolume);
+	RegisterInputConnection(&MinNumberOfPointsPerVoxel).SetCanHidePin(true).SetPinIsHidden(true);
+	RegisterInputConnection(&MaxNumberOfPointsPerVoxel).SetCanHidePin(true).SetPinIsHidden(true);
+	RegisterInputConnection(&RandomSeed).SetCanHidePin(true).SetPinIsHidden(true);
+	RegisterInputConnection(&IsoValue).SetCanHidePin(true).SetPinIsHidden(true);
+	RegisterInputConnection(&Spread).SetCanHidePin(true).SetPinIsHidden(true);
+	RegisterOutputConnection(&Points);
+}
+
+void FDenseUniformVolumeScatterDataflowNode::Evaluate(UE::Dataflow::FContext& Context, const FDataflowOutput* Out) const
+{
+	if (Out->IsA(&Points))
+	{
+		if (IsConnected(&FloatVolume))
+		{
+			const FDataflowFloatVolume& InFloatVolume = GetValue(Context, &FloatVolume);
+			const float InMinNumberOfPointsPerVoxel = GetValue(Context, &MinNumberOfPointsPerVoxel);
+			const float InMaxNumberOfPointsPerVoxel = GetValue(Context, &MaxNumberOfPointsPerVoxel);
+			const int32 InRandomSeed = GetValue(Context, &RandomSeed);
+			const float InIsoValue = GetValue(Context, &IsoValue);
+			const float InSpread = GetValue(Context, &Spread);
+
+			TArray<FVector> OutPoints;
+
+			// Compute point scatter in volume
+			InFloatVolume.DenseUniformVolumeScatter(
+				InMinNumberOfPointsPerVoxel,
+				InMaxNumberOfPointsPerVoxel,
+				InRandomSeed,
+				InIsoValue,
+				InSpread,
+				OutPoints);
+
+			SetValue(Context, MoveTemp(OutPoints), &Points);
+
+			return;
+		}
+
+		SetValue(Context, TArray<FVector>(), &Points);
+	}
+}
+
+#if WITH_EDITOR
+bool FDenseUniformVolumeScatterDataflowNode::CanDebugDrawViewMode(const FName& ViewModeName) const
+{
+	return ViewModeName == UE::Dataflow::FDataflowConstruction3DViewMode::Name;
+}
+
+void FDenseUniformVolumeScatterDataflowNode::DebugDraw(UE::Dataflow::FContext& Context, IDataflowDebugDrawInterface& DataflowRenderingInterface, const FDebugDrawParameters& DebugDrawParameters) const
+{
+	using namespace UE::Dataflow::Overlay;
+
+	if (DebugDrawParameters.bNodeIsSelected || DebugDrawParameters.bNodeIsPinned)
+	{
+		FString OverlayEvaluateResultString;
+
+		const FDataflowOutput* PointsOutput = FindOutput(&Points);
+		if (PointsOutput)
+		{
+			const TArray<FVector>& OutPoints = PointsOutput->GetValue(Context, Points);
+			//			const FDataflowFloatVolume& OutFloatVolume = this->GetOutputValue(Context, FloatVolume);
+
+			DataflowRenderingInterface.SetWorldPriority();
+			DataflowRenderingInterface.SetColor(Color);
+
+			for (const FVector& Point : OutPoints)
+			{
+				DataflowRenderingInterface.DrawPoint(Point);
+			}
+
+			OverlayEvaluateResultString.Appendf(TEXT("Number of Points: %d\n"), OutPoints.Num());
+		}
+
+		// ------------------------------------------------------------------------------------------------------
+		// Overlay
+		TArray<FString> OverlayStrArr;
+		OverlayStrArr.Add(BuildOverlayNodeInfoString(this));
+		OverlayStrArr.Add(OverlayEvaluateResultString);
+		FString OverlayStr = BuildOverlayFinalString(OverlayStrArr);
+		DataflowRenderingInterface.DrawOverlayText(OverlayStr);
+		// ------------------------------------------------------------------------------------------------------
+	}
+}
+#endif
+
+/* --------------------------------------------------------------------------------------------- */
+
+FNonUniformVolumeScatterDataflowNode::FNonUniformVolumeScatterDataflowNode(const UE::Dataflow::FNodeParameters& InParam, FGuid InGuid)
+	: FDataflowNode(InParam, InGuid)
+{
+	RegisterInputConnection(&FloatVolume);
+	RegisterInputConnection(&MinNumberOfPointsPerVoxel).SetCanHidePin(true).SetPinIsHidden(true);
+	RegisterInputConnection(&MaxNumberOfPointsPerVoxel).SetCanHidePin(true).SetPinIsHidden(true);
+	RegisterInputConnection(&RandomSeed).SetCanHidePin(true).SetPinIsHidden(true);
+	RegisterInputConnection(&IsoValue).SetCanHidePin(true).SetPinIsHidden(true);
+	RegisterInputConnection(&Spread).SetCanHidePin(true).SetPinIsHidden(true);
+	RegisterOutputConnection(&Points);
+}
+
+void FNonUniformVolumeScatterDataflowNode::Evaluate(UE::Dataflow::FContext& Context, const FDataflowOutput* Out) const
+{
+	if (Out->IsA(&Points))
+	{
+		if (IsConnected(&FloatVolume))
+		{
+			const FDataflowFloatVolume& InFloatVolume = GetValue(Context, &FloatVolume);
+			const float InMinNumberOfPointsPerVoxel = GetValue(Context, &MinNumberOfPointsPerVoxel);
+			const float InMaxNumberOfPointsPerVoxel = GetValue(Context, &MaxNumberOfPointsPerVoxel);
+			const int32 InRandomSeed = GetValue(Context, &RandomSeed);
+			const float InIsoValue = GetValue(Context, &IsoValue);
+			const float InSpread = GetValue(Context, &Spread);
+
+			TArray<FVector> OutPoints;
+
+			// Compute point scatter in volume
+			InFloatVolume.NonUniformVolumeScatter(
+				InMinNumberOfPointsPerVoxel,
+				InMaxNumberOfPointsPerVoxel,
+				InRandomSeed,
+				InIsoValue,
+				InSpread,
+				OutPoints);
+
+			SetValue(Context, MoveTemp(OutPoints), &Points);
+
+			return;
+		}
+
+		SetValue(Context, TArray<FVector>(), &Points);
+	}
+}
+
+#if WITH_EDITOR
+bool FNonUniformVolumeScatterDataflowNode::CanDebugDrawViewMode(const FName& ViewModeName) const
+{
+	return ViewModeName == UE::Dataflow::FDataflowConstruction3DViewMode::Name;
+}
+
+void FNonUniformVolumeScatterDataflowNode::DebugDraw(UE::Dataflow::FContext& Context, IDataflowDebugDrawInterface& DataflowRenderingInterface, const FDebugDrawParameters& DebugDrawParameters) const
+{
+	using namespace UE::Dataflow::Overlay;
+
+	if (DebugDrawParameters.bNodeIsSelected || DebugDrawParameters.bNodeIsPinned)
+	{
+		FString OverlayEvaluateResultString;
+
+		const FDataflowOutput* PointsOutput = FindOutput(&Points);
+		if (PointsOutput)
+		{
+			const TArray<FVector>& OutPoints = PointsOutput->GetValue(Context, Points);
+			//			const FDataflowFloatVolume& OutFloatVolume = this->GetOutputValue(Context, FloatVolume);
+
+			DataflowRenderingInterface.SetWorldPriority();
+			DataflowRenderingInterface.SetColor(Color);
+
+			for (const FVector& Point : OutPoints)
+			{
+				DataflowRenderingInterface.DrawPoint(Point);
+			}
+
+			OverlayEvaluateResultString.Appendf(TEXT("Number of Points: %d\n"), OutPoints.Num());
+		}
+
+		// ------------------------------------------------------------------------------------------------------
+		// Overlay
+		TArray<FString> OverlayStrArr;
+		OverlayStrArr.Add(BuildOverlayNodeInfoString(this));
+		OverlayStrArr.Add(OverlayEvaluateResultString);
+		FString OverlayStr = BuildOverlayFinalString(OverlayStrArr);
+		DataflowRenderingInterface.DrawOverlayText(OverlayStr);
+		// ------------------------------------------------------------------------------------------------------
+	}
+}
+#endif
+
+/* --------------------------------------------------------------------------------------------- */
+
+FVolumeSampleToPointsDataflowNode::FVolumeSampleToPointsDataflowNode(const UE::Dataflow::FNodeParameters& InParam, FGuid InGuid)
+: FDataflowNode(InParam, InGuid)
+{
+	RegisterInputConnection(&Volume);
+	RegisterInputConnection(&SamplePoints);
+	RegisterOutputConnection(&OutArray);
+	RegisterOutputConnection(&PointsCollection);
+
+	SetOutTypeFromSampleType();
+}
+
+void FVolumeSampleToPointsDataflowNode::SetOutTypeFromSampleType()
+{
+	if (FDataflowOutput* Output = FindOutput(&OutArray))
+	{
+		if (!Output->IsConnected())
+		{
+			Output->UnlockType();
+
+			switch (SampleType)
+			{
+			case EDataflowVolumeSampleType::Float:
+				SetOutputConcreteType<TArray<float>>(&OutArray);
+				break;
+			case EDataflowVolumeSampleType::Int:
+				SetOutputConcreteType<TArray<int32>>(&OutArray);
+				break;
+			case EDataflowVolumeSampleType::FloatVector:
+				SetOutputConcreteType<TArray<FVector>>(&OutArray);
+				break;
+			default:
+				break;
+			}
+
+			Output->LockType();
+		}
+		else
+		{
+			ensure(false);
+		}
+	}
+}
+
+void FVolumeSampleToPointsDataflowNode::OnPropertyChanged(UE::Dataflow::FContext& Context, const FPropertyChangedEvent& InPropertyChangedEvent)
+{
+	if (InPropertyChangedEvent.GetPropertyName() == GET_MEMBER_NAME_CHECKED(FVolumeSampleToPointsDataflowNode, SampleType))
+	{
+		SetOutTypeFromSampleType();
+	}
+}
+
+void FVolumeSampleToPointsDataflowNode::Evaluate(UE::Dataflow::FContext& Context, const FDataflowOutput* Out) const
+{
+	if (Out->IsA(&OutArray))
+	{
+		const FDataflowOutput* OutArrayOutput = FindOutput(&OutArray);
+
+		if (IsConnected(&Volume) && IsConnected(&SamplePoints))
+		{
+			// GetValue() for AnyType returns the storage type
+			const FDataflowVolume& InVolume = GetValue(Context, &Volume);
+			const TArray<FVector4>& InSamplePoints = GetValue(Context, &SamplePoints);
+			
+			FManagedArrayCollection OutPointsCollection = FManagedArrayCollection();
+			OutPointsCollection.AddGroup(FGeometryCollection::VerticesGroup);
+
+			const int32 NumPoints = InSamplePoints.Num();
+			TArray<FVector> Points;
+			Points.AddUninitialized(NumPoints);
+
+			for (int32 Idx = 0; Idx < NumPoints; ++Idx)
+			{
+				Points[Idx] = FVector(InSamplePoints[Idx]);
+			}
+
+			GeometryCollection::Facades::FPointsCollectionFacade PointsCollectionFacade(OutPointsCollection);
+
+			if (SampleType == EDataflowVolumeSampleType::Float)
+			{
+				if (const FDataflowFloatVolume* InFloatVolume = InVolume.Cast<FDataflowFloatVolume>())
+				{
+					TArray<float> OutValues;
+					OutValues.AddUninitialized(NumPoints);
+
+					InFloatVolume->VolumeSample(Points, OutValues);
+
+					PointsCollectionFacade.AddPointsWithFloatAttribute(Points, FName(*Attribute), OutValues);
+			
+					SetValue(Context, MoveTemp(OutPointsCollection), &PointsCollection);
+
+					if (OutArrayOutput)
+					{ 
+						OutArrayOutput->SetValue<TArray<float>>(MoveTemp(OutValues), Context);
+					}
+
+					return;
+				}
+				else
+				{
+//					Context.Error(TEXT("The connected volume is not a FloatVolume type"), this, Out);
+				}
+			}
+			else if (SampleType == EDataflowVolumeSampleType::Int)
+			{
+				// Input volume is an IntVolume
+				if (const FDataflowIntVolume* InIntVolume = InVolume.Cast<FDataflowIntVolume>())
+				{
+					TArray<int32> OutValues;
+					OutValues.AddUninitialized(Points.Num());
+
+					InIntVolume->VolumeSample(Points, OutValues);
+
+					PointsCollectionFacade.AddPointsWithIntAttribute(Points, FName(*Attribute), OutValues);
+
+					SetValue(Context, MoveTemp(OutPointsCollection), &PointsCollection);
+
+					if (OutArrayOutput)
+					{
+						OutArrayOutput->SetValue<TArray<int32>>(MoveTemp(OutValues), Context);
+					}
+
+					return;
+				}
+				else
+				{
+//					Context.Error(TEXT("The connected volume is not an IntVolume type"), this, Out);
+				}
+			}
+			else if (SampleType == EDataflowVolumeSampleType::FloatVector)
+			{
+				// Input volume is a FloatVectorVolume
+				if (const FDataflowFloatVectorVolume* InFloatVectorVolume = InVolume.Cast<FDataflowFloatVectorVolume>())
+				{
+					TArray<FVector> OutValues;
+					OutValues.AddUninitialized(Points.Num());
+
+					InFloatVectorVolume->VolumeSample(Points, OutValues);
+
+					PointsCollectionFacade.AddPointsWithVectorAttribute(Points, FName(*Attribute), OutValues);
+
+					SetValue(Context, MoveTemp(OutPointsCollection), &PointsCollection);
+
+					if (OutArrayOutput)
+					{
+						OutArrayOutput->SetValue<TArray<FVector>>(MoveTemp(OutValues), Context);
+					}
+
+					return;
+				}
+				else
+				{
+//					Context.Error(TEXT("The connected volume is not an FloatVectorVolume type"), this, Out);
+				}
+			}
+		}
+
+		if (OutArrayOutput)
+		{
+			OutArrayOutput->SetValue<TArray<FVector4>>(TArray<FVector4>(), Context);
+		}
+	}
+}
+
+#if WITH_EDITOR
+bool FVolumeSampleToPointsDataflowNode::CanDebugDrawViewMode(const FName& ViewModeName) const
+{
+	return ViewModeName == UE::Dataflow::FDataflowConstruction3DViewMode::Name;
+}
+
+void FVolumeSampleToPointsDataflowNode::DebugDraw(UE::Dataflow::FContext& Context, IDataflowDebugDrawInterface& DataflowRenderingInterface, const FDebugDrawParameters& DebugDrawParameters) const
+{
+	using namespace UE::Dataflow::Overlay;
+
+	if (DebugDrawParameters.bNodeIsSelected || DebugDrawParameters.bNodeIsPinned)
+	{
+		FString OverlayEvaluateResultString;
+
+		const TArray<FVector4>& InSamplePoints = GetValue(Context, &SamplePoints);
+
+		DataflowRenderingInterface.SetWorldPriority();
+		DataflowRenderingInterface.SetColor(Color);
+
+		for (const FVector4& Point : InSamplePoints)
+		{
+			DataflowRenderingInterface.DrawPoint(FVector(Point));
+		}
+
+		OverlayEvaluateResultString.Appendf(TEXT("Number of Points: %d\n"), InSamplePoints.Num());
+
+		// ------------------------------------------------------------------------------------------------------
+		// Overlay
+		TArray<FString> OverlayStrArr;
+		OverlayStrArr.Add(BuildOverlayNodeInfoString(this));
+		OverlayStrArr.Add(OverlayEvaluateResultString);
+		FString OverlayStr = BuildOverlayFinalString(OverlayStrArr);
+		DataflowRenderingInterface.DrawOverlayText(OverlayStr);
+		// ------------------------------------------------------------------------------------------------------
+	}
+}
+#endif
+
+/* --------------------------------------------------------------------------------------------- */
+
+namespace UE::Dataflow::Volume
+{
+	static void AdjustBoundingBox(FBox& VolumeBoundingBox, FVector& Extent, const float VoxelSize, const float Scale)
+	{
+
+		const FVector NewMin = VolumeBoundingBox.GetCenter() - (VolumeBoundingBox.GetCenter() - VolumeBoundingBox.Min) * Scale;
+		const FVector NewMax = VolumeBoundingBox.GetCenter() + (VolumeBoundingBox.Max - VolumeBoundingBox.GetCenter()) * Scale;
+
+		VolumeBoundingBox.Min = NewMin;
+		VolumeBoundingBox.Max = NewMax;
+
+		Extent = 2.f * VolumeBoundingBox.GetExtent();
+
+		float X = Extent.X / VoxelSize; Extent.X = ceil(X) * VoxelSize;
+		float Y = Extent.Y / VoxelSize; Extent.Y = ceil(Y) * VoxelSize;
+		float Z = Extent.Z / VoxelSize; Extent.Z = ceil(Z) * VoxelSize;
+
+		VolumeBoundingBox.Min = VolumeBoundingBox.GetCenter() - 0.5f * Extent;
+		VolumeBoundingBox.Max = VolumeBoundingBox.GetCenter() + 0.5f * Extent;
+	}
+
+	static TArray<FVector> GetSlicePoints(const FBox& VolumeBoundingBox,
+		const EDataflowVolumeSlicePlane Plane,
+		const float VoxelSize,
+		const float Offset)
+	{
+		TArray<FVector> Points;
+		int32 NumPoints;
+
+		FVector Extent = 2.f * VolumeBoundingBox.GetExtent();
+
+		FVector Min = VolumeBoundingBox.Min;
+		FVector Max = VolumeBoundingBox.Max;
+
+		if (Plane == EDataflowVolumeSlicePlane::XYPlane)
+		{
+			const int32 NumPointsX = int32(Extent.X / VoxelSize);
+			const int32 NumPointsY = int32(Extent.Y / VoxelSize);
+			NumPoints = NumPointsX * NumPointsY;
+
+			Points.AddUninitialized(NumPoints);
+
+			int32 Idx = 0;
+			for (int32 IdxX = 0; IdxX < NumPointsX; ++IdxX)
+			{
+				for (int32 IdxY = 0; IdxY < NumPointsY; ++IdxY)
+				{
+					Points[Idx].X = Min.X + float(IdxX) * VoxelSize;
+					Points[Idx].Y = Min.Y + float(IdxY) * VoxelSize;
+					Points[Idx++].Z = Max.Z - .5f * (1.f - Offset) * (Max.Z - Min.Z);
+				}
+			}
+		}
+		else if (Plane == EDataflowVolumeSlicePlane::YZPlane)
+		{
+			const int32 NumPointsY = int32(Extent.Y / VoxelSize);
+			const int32 NumPointsZ = int32(Extent.Z / VoxelSize);
+			NumPoints = NumPointsY * NumPointsZ;
+
+			Points.AddUninitialized(NumPoints);
+
+			int32 Idx = 0;
+			for (int32 IdxY = 0; IdxY < NumPointsY; ++IdxY)
+			{
+				for (int32 IdxZ = 0; IdxZ < NumPointsZ; ++IdxZ)
+				{
+					Points[Idx].X = Max.X - .5f * (1.f - Offset) * (Max.X - Min.X);
+					Points[Idx].Y = Min.Y + float(IdxY) * VoxelSize;
+					Points[Idx++].Z = Min.Z + float(IdxZ) * VoxelSize;
+				}
+			}
+		}
+		else if (Plane == EDataflowVolumeSlicePlane::ZXPlane)
+		{
+			const int32 NumPointsZ = int32(Extent.Z / VoxelSize);
+			const int32 NumPointsX = int32(Extent.X / VoxelSize);
+			NumPoints = NumPointsZ * NumPointsX;
+
+			Points.AddUninitialized(NumPoints);
+
+			int32 Idx = 0;
+			for (int32 IdxZ = 0; IdxZ < NumPointsZ; ++IdxZ)
+			{
+				for (int32 IdxX = 0; IdxX < NumPointsX; ++IdxX)
+				{
+					Points[Idx].X = Min.X + float(IdxX) * VoxelSize;
+					Points[Idx].Y = Max.Y - .5f * (1.f - Offset) * (Max.Y - Min.Y);
+					Points[Idx++].Z = Min.Z + float(IdxZ) * VoxelSize;
+				}
+			}
+		}
+
+		return Points;
+	}
+
+	//static void SetColorBar(const TArray<FLinearColor>& Spectrum, const FDataflowColorRamp ColorRamp)
+	//{
+	//	const int32 NumColors = Spectrum.Num();
+
+	//	for (int32 Idx = 0; Idx < NumColors; ++Idx)
+	//	{
+	//		ColorRamp.ColorRamp.GetCurves()[0].CurveToEdit->AddKey((float)Idx * 1.f / ((float)NumColors - 1.f), Spectrum[Idx].R);
+	//		ColorRamp.ColorRamp.GetCurves()[1].CurveToEdit->AddKey((float)Idx * 1.f / ((float)NumColors - 1.f), Spectrum[Idx].G);
+	//		ColorRamp.ColorRamp.GetCurves()[2].CurveToEdit->AddKey((float)Idx * 1.f / ((float)NumColors - 1.f), Spectrum[Idx].B);
+	//	}
+	//	ColorRamp.ColorRamp.GetCurves()[3].CurveToEdit->AddKey(0.f, 1.f);
+
+	//}
+
+	static TArray< FLinearColor> GetInfraRedColorSpectrum()
+	{
+		const int32 NumColors = 7;
+		TArray<FLinearColor> Spectrum;
+		Spectrum.AddUninitialized(NumColors);
+
+		Spectrum[6] = FLinearColor::Red;
+		Spectrum[5] = FLinearColor(FColor(255, 128, 0, 1)); // Orange
+		Spectrum[4] = FLinearColor::Yellow;
+		Spectrum[3] = FLinearColor::Green;
+		Spectrum[2] = FLinearColor::Blue;
+		Spectrum[1] = FLinearColor(FColor(64, 0, 128, 1)); // Indigo
+		Spectrum[0] = FLinearColor(FColor(128, 0, 128, 1)); // Violet
+
+		return Spectrum;
+	}
+}
+
+FVolumeSliceDataflowNode::FVolumeSliceDataflowNode(const UE::Dataflow::FNodeParameters& InParam, FGuid InGuid)
+	: FDataflowNode(InParam, InGuid)
+{
+	static const FName TypeDependencyGroup("Main");
+	RegisterInputConnection(&Volume)
+		.SetTypeDependencyGroup(TypeDependencyGroup);
+	RegisterOutputConnection(&Volume, &Volume)
+		.SetTypeDependencyGroup(TypeDependencyGroup);
+
+	// Set ColorRamp
+	const int32 NumColors = 7;
+	TArray<FLinearColor> Spectrum;
+	Spectrum.AddUninitialized(NumColors);
+
+	Spectrum[6] = FLinearColor::Red;
+	Spectrum[5] = FLinearColor(FColor(255, 128, 0, 1)); // Orange
+	Spectrum[4] = FLinearColor::Yellow;
+	Spectrum[3] = FLinearColor::Green;
+	Spectrum[2] = FLinearColor::Blue;
+	Spectrum[1] = FLinearColor(FColor(64, 0, 128, 1)); // Indigo
+	Spectrum[0] = FLinearColor(FColor(128, 0, 128, 1)); // Violet
+
+	for (int32 Idx = 0; Idx < NumColors; ++Idx)
+	{
+		ColorRamp.ColorRamp.GetCurves()[0].CurveToEdit->AddKey((float)Idx * 1.f / ((float)NumColors - 1.f), Spectrum[Idx].R);
+		ColorRamp.ColorRamp.GetCurves()[1].CurveToEdit->AddKey((float)Idx * 1.f / ((float)NumColors - 1.f), Spectrum[Idx].G);
+		ColorRamp.ColorRamp.GetCurves()[2].CurveToEdit->AddKey((float)Idx * 1.f / ((float)NumColors - 1.f), Spectrum[Idx].B);
+	}
+	ColorRamp.ColorRamp.GetCurves()[3].CurveToEdit->AddKey(0.f, 1.f);
+
+	//
+	ColorRamp.ColorRamp.OnColorCurveChangedDelegate.AddRaw(this, &FVolumeSliceDataflowNode::OnColorCurveChanged);
+}
+
+void FVolumeSliceDataflowNode::OnColorCurveChanged(TArray<FRichCurve*> Curves)
+{
+	Invalidate();
+}
+
+void FVolumeSliceDataflowNode::Evaluate(UE::Dataflow::FContext& Context, const FDataflowOutput* Out) const
+{
+	if (Out->IsA(&Volume))
+	{
+		ForwardInput(Context, &Volume, &Volume);
+	}
+}
+
+void FVolumeSliceDataflowNode::OnPropertyChanged(UE::Dataflow::FContext& Context, const FPropertyChangedEvent& InPropertyChangedEvent)
+{
+#if WITH_EDITOR
+	const FName PropertyName = InPropertyChangedEvent.GetPropertyName();
+
+	if (PropertyName == GET_MEMBER_NAME_CHECKED(FVolumeSliceDataflowNode, ColorRampType))
+	{
+		TArray<FLinearColor> Spectrum;
+		int32 NumColors = 0;
+
+		if (ColorRampType == EDataflowVolumeSliceRamp::InfraRed)
+		{
+			NumColors = 7;
+			Spectrum.AddUninitialized(NumColors);
+
+			Spectrum[0] = FLinearColor(FColor(128, 0, 128, 1)); // Violet
+			Spectrum[1] = FLinearColor(FColor(64, 0, 128, 1)); // Indigo
+			Spectrum[2] = FLinearColor::Blue;
+			Spectrum[3] = FLinearColor::Green;
+			Spectrum[4] = FLinearColor::Yellow;
+			Spectrum[5] = FLinearColor(FColor(255, 128, 0, 1)); // Orange
+			Spectrum[6] = FLinearColor::Red;
+		}
+		else if (ColorRampType == EDataflowVolumeSliceRamp::WhiteToRed)
+		{
+			NumColors = 2;
+			Spectrum.AddUninitialized(NumColors);
+
+			Spectrum[0] = FLinearColor::White;
+			Spectrum[1] = FLinearColor::Red;
+		}
+		else if (ColorRampType == EDataflowVolumeSliceRamp::BlackBody)
+		{
+			NumColors = 4;
+			Spectrum.AddUninitialized(NumColors);
+
+			Spectrum[0] = FLinearColor::White;
+			Spectrum[1] = FLinearColor::Yellow;
+			Spectrum[2] = FLinearColor::Red;
+			Spectrum[3] = FLinearColor::Black;
+		}
+		else if (ColorRampType == EDataflowVolumeSliceRamp::Grayscale)
+		{
+			NumColors = 2;
+			Spectrum.AddUninitialized(NumColors);
+
+			Spectrum[0] = FLinearColor::Black;
+			Spectrum[1] = FLinearColor::White;
+		}
+		else if (ColorRampType == EDataflowVolumeSliceRamp::Custom)
+		{
+			NumColors = 2;
+			Spectrum.AddUninitialized(NumColors);
+
+			Spectrum[0] = FLinearColor::Red;
+			Spectrum[1] = FLinearColor::Blue;
+		}
+
+		// Update ColorRamp
+		for (int32 Idx = 0; Idx < 4; ++Idx)
+		{
+			ColorRamp.ColorRamp.GetCurves()[Idx].CurveToEdit->Reset();
+		}
+
+		for (int32 Idx = 0; Idx < NumColors; ++Idx)
+		{
+			ColorRamp.ColorRamp.GetCurves()[0].CurveToEdit->AddKey((float)Idx * 1.f / ((float)NumColors - 1.f), Spectrum[Idx].R);
+			ColorRamp.ColorRamp.GetCurves()[1].CurveToEdit->AddKey((float)Idx * 1.f / ((float)NumColors - 1.f), Spectrum[Idx].G);
+			ColorRamp.ColorRamp.GetCurves()[2].CurveToEdit->AddKey((float)Idx * 1.f / ((float)NumColors - 1.f), Spectrum[Idx].B);
+		}
+		ColorRamp.ColorRamp.GetCurves()[3].CurveToEdit->AddKey(0.f, 1.f);
+
+	}
+#endif
+}
+
+#if WITH_EDITOR
+bool FVolumeSliceDataflowNode::CanDebugDrawViewMode(const FName& ViewModeName) const
+{
+	return ViewModeName == UE::Dataflow::FDataflowConstruction3DViewMode::Name;
+}
+
+void FVolumeSliceDataflowNode::DebugDraw(UE::Dataflow::FContext& Context, IDataflowDebugDrawInterface& DataflowRenderingInterface, const FDebugDrawParameters& DebugDrawParameters) const
+{
+	using namespace UE::Dataflow::Overlay;
+
+	if (DebugDrawParameters.bNodeIsSelected || DebugDrawParameters.bNodeIsPinned)
+	{
+		if (IsConnected(&Volume))
+		{
+			if (bVisualize)
+			{
+				const FDataflowVolume& InVolume = GetValue(Context, &Volume);
+				if (const FDataflowFloatVolume* InFloatVolume = InVolume.Cast<FDataflowFloatVolume>())
+				{
+					const float VoxelSize = InFloatVolume->GetVoxelSize()[0];
+					FBox VolumeBoundingBox = InFloatVolume->GetVolumeBoundingBox();
+
+					FVector Extent;
+					UE::Dataflow::Volume::AdjustBoundingBox(VolumeBoundingBox, Extent, VoxelSize, BoundingBoxScale);
+
+					TArray<FVector> SamplePoints = UE::Dataflow::Volume::GetSlicePoints(VolumeBoundingBox, Plane, VoxelSize, Offset);
+
+					// Sample points in the volume
+					TArray<float> SampledValues;
+					SampledValues.AddUninitialized(SamplePoints.Num());
+
+					InFloatVolume->VolumeSample(SamplePoints, SampledValues);
+
+					DataflowRenderingInterface.SetWorldPriority();
+
+					// Draw BoundingBox
+					DataflowRenderingInterface.SetColor(FLinearColor::Gray);
+					DataflowRenderingInterface.DrawBox(0.5 * Extent, FQuat::Identity, VolumeBoundingBox.GetCenter(), 1.f);
+
+					float Progress;
+
+					for (int32 Idx = 0; Idx < SamplePoints.Num(); ++Idx)
+					{
+						float SampledValue = FMath::Clamp(SampledValues[Idx], Range.X, Range.Y);
+
+						FLinearColor Color;
+						if (ColorRampType == EDataflowVolumeSliceRamp::WhiteToRed)
+						{
+							Progress = 1.f - (Range.Y - SampledValue) / (Range.Y - Range.X);
+//							Color = FLinearColor::LerpUsingHSV(FLinearColor::White, FLinearColor::Red, Progress);
+//							Color = ColorRamp.ColorRamp.GetLinearColorValue(Progress);
+						}
+						else if (ColorRampType == EDataflowVolumeSliceRamp::Grayscale)
+						{
+							Progress = 1.f - (Range.Y - SampledValue) / (Range.Y - Range.X);
+//							Color = FLinearColor::LerpUsingHSV(FLinearColor::Black, FLinearColor::White, Progress);
+//							Color = ColorRamp.ColorRamp.GetLinearColorValue(Progress);
+						}
+						else if (ColorRampType == EDataflowVolumeSliceRamp::Custom)
+						{
+							Progress = 1.f - (Range.Y - SampledValue) / (Range.Y - Range.X);
+//							Color = FLinearColor::LerpUsingHSV(ColorA, ColorB, Progress);
+//							Color = ColorRamp.ColorRamp.GetLinearColorValue(Progress);
+						}
+						else if (ColorRampType == EDataflowVolumeSliceRamp::InfraRed)
+						{
+							const int32 NumColors = 7;
+							//TArray<FLinearColor> Spectrum; Spectrum.AddUninitialized(NumColors);
+
+							//Spectrum[6] = FLinearColor::Red;
+							//Spectrum[5] = FLinearColor(FColor(255, 128, 0, 1)); // Orange
+							//Spectrum[4] = FLinearColor::Yellow;
+							//Spectrum[3] = FLinearColor::Green;
+							//Spectrum[2] = FLinearColor::Blue;
+							//Spectrum[1] = FLinearColor(FColor(64, 0, 128, 1)); // Indigo
+							//Spectrum[0] = FLinearColor(FColor(128, 0, 128, 1)); // Violet
+
+							//if ((SampledValue - Range.X) < 1e-2)
+							//{
+							//	SampledValue += 1e-2;
+							//}
+							//else if ((Range.Y - SampledValue) < 1e-2)
+							//{
+							//	SampledValue -= 1e-2;
+							//}
+
+							//const float Delta = (Range.Y - Range.X) / float(NumColors - 1);
+							//const float B = floorf((SampledValue - Range.X) / Delta);
+							//const float NewMin = Range.X + B * Delta;
+							//const float NewMax = Range.X + (B + 1) * Delta;
+
+//							Progress = 1.f - (NewMax - SampledValue) / Delta;
+							Progress = 1.f - (Range.Y - SampledValue) / (Range.Y - Range.X);
+//							Color = FLinearColor::LerpUsingHSV(Spectrum[int32(B)], Spectrum[int32(B + 1)], Progress);
+						}
+						else if (ColorRampType == EDataflowVolumeSliceRamp::BlackBody)
+						{
+							const int32 NumColors = 4;
+							//TArray<FLinearColor> Spectrum; Spectrum.AddUninitialized(NumColors);
+
+							//Spectrum[3] = FLinearColor::White;
+							//Spectrum[2] = FLinearColor::Yellow;
+							//Spectrum[1] = FLinearColor::Red;
+							//Spectrum[0] = FLinearColor::Black;
+
+							//const float Delta = (Range.Y - Range.X) / float(NumColors - 1);
+							//const float B = floorf((SampledValue - Range.X) / Delta);
+							//const float NewMin = Range.X + B * Delta;
+							//const float NewMax = Range.X + (B + 1) * Delta;
+
+//							Progress = 1.f - (NewMax - SampledValue) / Delta;
+							Progress = 1.f - (Range.Y - SampledValue) / (Range.Y - Range.X);
+							//Color = Spectrum[int32(B)];
+						}
+						Progress = 1.f - (Range.Y - SampledValue) / (Range.Y - Range.X);
+						if (bInvertColorRamp)
+						{
+							Progress = 1.f - Progress;
+						}
+						Color = ColorRamp.ColorRamp.GetLinearColorValue(Progress);
+
+						DataflowRenderingInterface.SetColor(Color);
+						DataflowRenderingInterface.DrawPoint(SamplePoints[Idx]);
+					}
+				}
+				else if (const FDataflowFloatVectorVolume* InFloatVectorVolume = InVolume.Cast<FDataflowFloatVectorVolume>())
+				{
+					const float VoxelSize = InFloatVectorVolume->GetVoxelSize()[0];
+					FBox VolumeBoundingBox = InFloatVectorVolume->GetVolumeBoundingBox();
+
+					FVector Extent;
+					UE::Dataflow::Volume::AdjustBoundingBox(VolumeBoundingBox, Extent, VoxelSize, BoundingBoxScale);
+
+					TArray<FVector> SamplePoints = UE::Dataflow::Volume::GetSlicePoints(VolumeBoundingBox, Plane, VoxelSize, Offset);
+
+					// Sample points in the volume
+					TArray<FVector> SampledValues;
+					SampledValues.AddUninitialized(SamplePoints.Num());
+
+					InFloatVectorVolume->VolumeSample(SamplePoints, SampledValues);
+
+					DataflowRenderingInterface.SetWorldPriority();
+					DataflowRenderingInterface.SetPointSize(PointScale);
+					DataflowRenderingInterface.SetLineWidth(LineWidth);
+
+					// Draw BoundingBox
+					DataflowRenderingInterface.SetColor(FLinearColor::Gray);
+					DataflowRenderingInterface.DrawBox(0.5 * Extent, FQuat::Identity, VolumeBoundingBox.GetCenter(), 1.f);
+
+					// Draw values
+					if (InFloatVectorVolume->GetVectorType() == "Covariant")
+					{
+						for (int32 Idx = 0; Idx < SamplePoints.Num(); ++Idx)
+						{
+							if (SampledValues[Idx].Length() > MinLength)
+							{
+								DataflowRenderingInterface.SetColor(VectorColor);
+								DataflowRenderingInterface.DrawLine(SamplePoints[Idx], SamplePoints[Idx] + VectorScale * SampledValues[Idx]);
+
+								DataflowRenderingInterface.SetColor(PointColor);
+								DataflowRenderingInterface.DrawPoint(SamplePoints[Idx]);
+							}
+						}
+					}
+				}
+			}
+		}
+	}
+}
+#endif
+
+/* --------------------------------------------------------------------------------------------- */
+
+FVolumeAnalysisDataflowNode::FVolumeAnalysisDataflowNode(const UE::Dataflow::FNodeParameters& InParam, FGuid InGuid)
+	: FDataflowNode(InParam, InGuid)
+{
+	RegisterInputConnection(&Volume);
+	RegisterInputConnection(&MaskVolume);
+	RegisterOutputConnection(&OutVolume);
+	RegisterOutputConnection(&MaskVolume, &MaskVolume);
+
+	SetOutTypeFromOperator();
+}
+
+void FVolumeAnalysisDataflowNode::SetOutTypeFromOperator()
+{
+	if (FDataflowOutput* Output = FindOutput(&OutVolume))
+	{
+		if (!Output->IsConnected())
+		{
+			Output->UnlockType();
+
+			if (Operator == EDataflowVolumeAnalysisOperator::Gradient ||
+				Operator == EDataflowVolumeAnalysisOperator::ClosestPoint ||
+				Operator == EDataflowVolumeAnalysisOperator::Curl ||
+				Operator == EDataflowVolumeAnalysisOperator::Normalize)
+			{
+				SetOutputConcreteType<FDataflowFloatVectorVolume>(&OutVolume);
+			}
+			else if (Operator == EDataflowVolumeAnalysisOperator::Curvature ||
+				Operator == EDataflowVolumeAnalysisOperator::Laplacian ||
+				Operator == EDataflowVolumeAnalysisOperator::Divergence ||
+				Operator == EDataflowVolumeAnalysisOperator::Magnitude)
+			{
+				SetOutputConcreteType<FDataflowFloatVolume>(&OutVolume);
+			}
+
+			Output->LockType();
+		}
+		else
+		{
+			ensure(false);
+		}
+	}
+}
+
+void FVolumeAnalysisDataflowNode::OnPropertyChanged(UE::Dataflow::FContext& Context, const FPropertyChangedEvent& InPropertyChangedEvent)
+{
+	if (InPropertyChangedEvent.GetPropertyName() == GET_MEMBER_NAME_CHECKED(FVolumeAnalysisDataflowNode, Operator))
+	{
+		SetOutTypeFromOperator();
+	}
+}
+
+void FVolumeAnalysisDataflowNode::Evaluate(UE::Dataflow::FContext& Context, const FDataflowOutput* Out) const
+{
+	if (Out->IsA(&OutVolume))
+	{
+		if (IsConnected(&Volume))
+		{
+			const FDataflowOutput* VolumeOutput = FindOutput(&OutVolume);
+			
+			if (VolumeOutput)
+			{
+				FDataflowBoolVolume* InMaskVolumePtr = nullptr;
+				if (IsConnected(&MaskVolume))
+				{
+					FDataflowBoolVolume InMaskVolume = GetValue(Context, &MaskVolume);
+					InMaskVolumePtr = &InMaskVolume;
+				}
+
+				const FDataflowVolume& InVolume = GetValue(Context, &Volume);
+				if (const FDataflowFloatVolume* InFloatVolume = InVolume.Cast<FDataflowFloatVolume>())
+				{
+					// Scalar -> Vector
+					if (Operator == EDataflowVolumeAnalysisOperator::Gradient)
+					{
+						FDataflowFloatVectorVolume GradientVolume = InFloatVolume->ComputeGradient(OutputName, CustomName, InMaskVolumePtr);
+						VolumeOutput->SetValue<FDataflowFloatVectorVolume>(MoveTemp(GradientVolume), Context);
+					}
+					// Scalar -> Scalar
+					else if (Operator == EDataflowVolumeAnalysisOperator::Curvature)
+					{
+						FDataflowFloatVolume CurvatureVolume = InFloatVolume->ComputeCurvature(OutputName, CustomName, InMaskVolumePtr);
+						VolumeOutput->SetValue<FDataflowFloatVolume>(MoveTemp(CurvatureVolume), Context);
+					}
+					// Scalar -> Scalar
+					else if (Operator == EDataflowVolumeAnalysisOperator::Laplacian)
+					{
+						FDataflowFloatVolume LaplacianVolume = InFloatVolume->ComputeLaplacian(OutputName, CustomName, InMaskVolumePtr);
+						VolumeOutput->SetValue<FDataflowFloatVolume>(MoveTemp(LaplacianVolume), Context);
+					}
+					// Scalar -> Vector
+					else if (Operator == EDataflowVolumeAnalysisOperator::ClosestPoint)
+					{
+						FDataflowFloatVectorVolume ClosestPointVolume = InFloatVolume->ComputeClosestPoint(OutputName, CustomName, InMaskVolumePtr);
+						VolumeOutput->SetValue<FDataflowFloatVectorVolume>(MoveTemp(ClosestPointVolume), Context);
+					}
+				}
+				else if (const FDataflowFloatVectorVolume* InFloatVectorVolume = InVolume.Cast<FDataflowFloatVectorVolume>())
+				{
+					// Vector -> Scalar
+					if (Operator == EDataflowVolumeAnalysisOperator::Divergence)
+					{
+						FDataflowFloatVolume DivergenceVolume = InFloatVectorVolume->ComputeDivergence(OutputName, CustomName);
+//						FDataflowFloatVolume DivergenceVolume = InFloatVectorVolume->ComputeAnalysis<FDataflowFloatVolume, "Divergence">(OutputName, CustomName, InMaskVolumePtr);
+						VolumeOutput->SetValue<FDataflowFloatVolume>(MoveTemp(DivergenceVolume), Context);
+					}
+					// Vector -> Vector
+					else if (Operator == EDataflowVolumeAnalysisOperator::Curl)
+					{
+						FDataflowFloatVectorVolume CurlVolume = InFloatVectorVolume->ComputeCurl(OutputName, CustomName);
+//						FDataflowFloatVectorVolume CurlVolume = InFloatVectorVolume->ComputeAnalysis<FDataflowFloatVectorVolume, "Curl">(OutputName, CustomName, InMaskVolumePtr);
+						VolumeOutput->SetValue<FDataflowFloatVectorVolume>(MoveTemp(CurlVolume), Context);
+					}
+					// Vector -> Scalar
+					else if (Operator == EDataflowVolumeAnalysisOperator::Magnitude)
+					{
+						FDataflowFloatVolume LengthVolume = InFloatVectorVolume->ComputeMagnitude(OutputName, CustomName);
+						VolumeOutput->SetValue<FDataflowFloatVolume>(MoveTemp(LengthVolume), Context);
+					}
+					// Vector -> Vector
+					else if (Operator == EDataflowVolumeAnalysisOperator::Normalize)
+					{
+						FDataflowFloatVectorVolume NormalVolume = InFloatVectorVolume->ComputeNormalize(OutputName, CustomName);
+						VolumeOutput->SetValue<FDataflowFloatVectorVolume>(MoveTemp(NormalVolume), Context);
+					}
+				}
+			}
+		}
+	}
+	else if (Out->IsA(&MaskVolume))
+	{
+		if (IsConnected(&MaskVolume))
+		{
+			SafeForwardInput(Context, &MaskVolume, &MaskVolume);
+		}
+	}
+}
+
+/* --------------------------------------------------------------------------------------------- */
+
+FVolumeCombineDataflowNode::FVolumeCombineDataflowNode(const UE::Dataflow::FNodeParameters& InParam, FGuid InGuid)
+	: FDataflowNode(InParam, InGuid)
+{
+	RegisterInputConnection(&FloatVolumeA);
+	RegisterInputConnection(&FloatVolumeB);
+	RegisterInputConnection(&MultiplierA).SetCanHidePin(true).SetPinIsHidden(true);
+	RegisterInputConnection(&MultiplierB).SetCanHidePin(true).SetPinIsHidden(true);
+	RegisterOutputConnection(&FloatVolume);
+
+	DebugDrawRenderSettings.RenderType = EDataflowDebugDrawRenderType::Wireframe;
+	DebugDrawRenderSettings.bTranslucent = false;
+	DebugDrawRenderSettings.Color = FLinearColor(FColor::Cyan);
+	DebugDrawRenderSettings.LineWidthMultiplier = 1.5f;
+}
+
+void FVolumeCombineDataflowNode::Evaluate(UE::Dataflow::FContext& Context, const FDataflowOutput* Out) const
+{
+	if (Out->IsA(&FloatVolume))
+	{
+		if (IsConnected(&FloatVolumeA) && IsConnected(&FloatVolumeB))
+		{
+			FDataflowFloatVolume InFloatVolumeA = GetValue(Context, &FloatVolumeA);
+			FDataflowFloatVolume InFloatVolumeB = GetValue(Context, &FloatVolumeB);
+
+			if (bFlipInputs)
+			{
+				InFloatVolumeA = GetValue(Context, &FloatVolumeB);
+				InFloatVolumeB = GetValue(Context, &FloatVolumeA);
+			}
+
+			const float InMultiplierA = GetValue(Context, &MultiplierA);
+			const float InMultiplierB = GetValue(Context, &MultiplierB);
+
+			FDataflowFloatVolume OutVolume = InFloatVolumeA.VolumeCombine(
+				InFloatVolumeB,
+				Operation,
+				InMultiplierA,
+				InMultiplierB,
+				Resample,
+				Interpolation,
+				bPruneDegenerateTiles,
+				PruneTol);
+
+			SetValue(Context, MoveTemp(OutVolume), &FloatVolume);
+			return;
+		}
+
+		SetValue(Context, FDataflowFloatVolume(), &FloatVolume);
+	}
+}
+
+#if WITH_EDITOR
+bool FVolumeCombineDataflowNode::CanDebugDrawViewMode(const FName& ViewModeName) const
+{
+	return ViewModeName == UE::Dataflow::FDataflowConstruction3DViewMode::Name;
+}
+
+void FVolumeCombineDataflowNode::DebugDraw(UE::Dataflow::FContext& Context, IDataflowDebugDrawInterface& DataflowRenderingInterface, const FDebugDrawParameters& DebugDrawParameters) const
+{
+	using namespace UE::Dataflow::Overlay;
+
+	if (DebugDrawParameters.bNodeIsSelected || DebugDrawParameters.bNodeIsPinned)
+	{
+		FString OverlayEvaluateResultString;
+
+		const FDataflowOutput* FloatVolumeOutput = FindOutput(&FloatVolume);
+		if (FloatVolumeOutput)
+		{
+			const FDataflowFloatVolume& OutFloatVolume = FloatVolumeOutput->GetValue(Context, FloatVolume);
+
+			DebugDrawRenderSettings.SetDebugDrawSettings(DataflowRenderingInterface);
+
+			TArray<FBox> ActiveVoxels;
+			OutFloatVolume.GetActiveVoxels(ActiveVoxels);
+
+			for (const FBox& Voxel : ActiveVoxels)
+			{
+				if (bDrawVoxelCenters)
+				{
+					DataflowRenderingInterface.DrawPoint(0.5 * FVector(Voxel.Max + Voxel.Min));
+				}
+
+				if (bDrawVoxels)
+				{
+					DataflowRenderingInterface.DrawBox(0.5 * FVector(Voxel.Max - Voxel.Min), FQuat::Identity, 0.5 * FVector(Voxel.Min + Voxel.Max), 1.0);
+				}
+			}
+
+			OverlayEvaluateResultString = OutFloatVolume.VolumeInfo();
+		}
+
+		// ------------------------------------------------------------------------------------------------------
+		// Overlay
+		TArray<FString> OverlayStrArr;
+		OverlayStrArr.Add(BuildOverlayNodeInfoString(this));
+		OverlayStrArr.Add(OverlayEvaluateResultString);
+		FString OverlayStr = BuildOverlayFinalString(OverlayStrArr);
+		DataflowRenderingInterface.DrawOverlayText(OverlayStr);
+		// ------------------------------------------------------------------------------------------------------
+	}
+}
+#endif
+
+/* --------------------------------------------------------------------------------------------- */
+
+
+
+
+
+
+
+
+
+
+
+
+
+
diff --git a/Engine/Plugins/Experimental/DataflowVolumetricPlugin/Source/DataflowVolumeNodes/Private/Dataflow/DataflowVolumeNodes.h b/Engine/Plugins/Experimental/DataflowVolumetricPlugin/Source/DataflowVolumeNodes/Private/Dataflow/DataflowVolumeNodes.h
new file mode 100644
--- /dev/null
+++ b/Engine/Plugins/Experimental/DataflowVolumetricPlugin/Source/DataflowVolumeNodes/Private/Dataflow/DataflowVolumeNodes.h
@@ -0,0 +1,863 @@
+// Copyright Epic Games, Inc. All Rights Reserved.
+
+#pragma once 
+
+#include "Dataflow/DataflowEngine.h"
+#include "Dataflow/DataflowFloatVolume.h"
+#include "Dataflow/DataflowFloatVectorVolume.h"
+#include "Dataflow/DataflowIntVolume.h"
+#include "Dataflow/DataflowBoolVolume.h"
+#include "Dataflow/DataflowVolumeAnyType.h"
+#include "Dataflow/DataflowDebugDraw.h"
+#include "Dataflow/DataflowColorRamp.h"
+
+#include "DataflowVolumeNodes.generated.h"
+
+/**
+ * Make a sphere SDF
+ */
+USTRUCT(meta = (Experimental, DataflowVolume))
+struct FMakeSDFSphereDataflowNode : public FDataflowNode
+{
+	GENERATED_USTRUCT_BODY()
+	DATAFLOW_NODE_DEFINE_INTERNAL(FMakeSDFSphereDataflowNode, "MakeSDFSphere", "Volume|Generators", "")
+
+private:
+	/** The size of voxels in the output VDB volume */
+	UPROPERTY(EditAnywhere, Category = "SDF", meta = (DataflowInput, ClampMin = 0.1f));
+	float VoxelSize = 1.f;
+
+	/** Sphere radius */
+	UPROPERTY(EditAnywhere, Category = "Sphere", meta = (DataflowInput, ClampMin = 0.1f));
+	float Radius = 3.f;
+
+	/** Sphere center */
+	UPROPERTY(EditAnywhere, Category = "Sphere", meta = (DataflowInput));
+	FVector Center = FVector(0);
+
+	UPROPERTY(EditAnywhere, Category = "Debug Draw", meta = (ShowOnlyInnerProperties))
+	FDataflowNodeDebugDrawSettings DebugDrawRenderSettings;
+
+	UPROPERTY(EditAnywhere, Category = "Debug Draw")
+	bool bDrawVoxels = true;
+
+	UPROPERTY(EditAnywhere, Category = "Debug Draw")
+	bool bDrawVoxelCenters = false;
+
+	/** SDF volume */
+	UPROPERTY(meta = (DataflowOutput))
+	FDataflowFloatVolume FloatVolume;
+
+	virtual void Evaluate(UE::Dataflow::FContext& Context, const FDataflowOutput* Out) const override;
+
+#if WITH_EDITOR
+	virtual bool CanDebugDraw() const override
+	{
+		return true;
+	}
+	virtual bool CanDebugDrawViewMode(const FName& ViewModeName) const override;
+	virtual void DebugDraw(UE::Dataflow::FContext& Context, IDataflowDebugDrawInterface& DataflowRenderingInterface, const FDebugDrawParameters& DebugDrawParameters) const override;
+#endif
+
+public:
+	FMakeSDFSphereDataflowNode(const UE::Dataflow::FNodeParameters& InParam, FGuid InGuid = FGuid::NewGuid());
+};
+
+/**
+ * Make a cube SDF
+ */
+USTRUCT(meta = (Experimental, DataflowVolume))
+struct FMakeSDFCubeDataflowNode : public FDataflowNode
+{
+	GENERATED_USTRUCT_BODY()
+	DATAFLOW_NODE_DEFINE_INTERNAL(FMakeSDFCubeDataflowNode, "MakeSDFCube", "Volume|Generators", "")
+
+private:
+	/** Scale */
+	UPROPERTY(EditAnywhere, Category = "Cube", meta = (DataflowInput));
+	float Scale = 1.f;
+
+	/** Center */
+	UPROPERTY(EditAnywhere, Category = "Cube", meta = (DataflowInput));
+	FVector Center = FVector(0.0);
+
+	/** The size of voxels in the output VDB volume */
+	UPROPERTY(EditAnywhere, Category = "SDF", meta = (DataflowInput, ClampMin = 0.1f));
+	float VoxelSize = 1.f;
+
+	UPROPERTY(EditAnywhere, Category = "Debug Draw", meta = (ShowOnlyInnerProperties))
+	FDataflowNodeDebugDrawSettings DebugDrawRenderSettings;
+
+	UPROPERTY(EditAnywhere, Category = "Debug Draw")
+	bool bDrawVoxels = true;
+
+	UPROPERTY(EditAnywhere, Category = "Debug Draw")
+	bool bDrawVoxelCenters = false;
+
+	/** SDF volume */
+	UPROPERTY(meta = (DataflowOutput))
+	FDataflowFloatVolume FloatVolume;
+
+	virtual void Evaluate(UE::Dataflow::FContext& Context, const FDataflowOutput* Out) const override;
+
+#if WITH_EDITOR
+	virtual bool CanDebugDraw() const override
+	{
+		return true;
+	}
+	virtual bool CanDebugDrawViewMode(const FName& ViewModeName) const override;
+	virtual void DebugDraw(UE::Dataflow::FContext& Context, IDataflowDebugDrawInterface& DataflowRenderingInterface, const FDebugDrawParameters& DebugDrawParameters) const override;
+#endif
+
+public:
+	FMakeSDFCubeDataflowNode(const UE::Dataflow::FNodeParameters& InParam, FGuid InGuid = FGuid::NewGuid());
+};
+
+/**
+ * Make a platonic solid SDF
+ */
+USTRUCT(meta = (Experimental, DataflowVolume))
+struct FMakeSDFPlatonicSolidDataflowNode : public FDataflowNode
+{
+	GENERATED_USTRUCT_BODY()
+	DATAFLOW_NODE_DEFINE_INTERNAL(FMakeSDFPlatonicSolidDataflowNode, "MakeSDFPlatonicSolid", "Volume|Generators", "")
+
+private:
+	/** Platonic solid type */
+	UPROPERTY(EditAnywhere, Category = "Platonic Solid");
+	EDataflowVolumePlatonicSolidType PlatonicSolidType = EDataflowVolumePlatonicSolidType::Tetrahedron;
+
+	/** Scale */
+	UPROPERTY(EditAnywhere, Category = "Platonic Solid", meta = (DataflowInput));
+	float Scale = 1.f;
+
+	/** Center */
+	UPROPERTY(EditAnywhere, Category = "Platonic Solid", meta = (DataflowInput));
+	FVector Center = FVector(0.0);
+
+	/** The size of voxels in the output VDB volume */
+	UPROPERTY(EditAnywhere, Category = "SDF", meta = (DataflowInput, ClampMin = 0.1f));
+	float VoxelSize = 1.f;
+
+	UPROPERTY(EditAnywhere, Category = "Debug Draw", meta = (ShowOnlyInnerProperties))
+	FDataflowNodeDebugDrawSettings DebugDrawRenderSettings;
+
+	UPROPERTY(EditAnywhere, Category = "Debug Draw")
+	bool bDrawVoxels = true;
+
+	UPROPERTY(EditAnywhere, Category = "Debug Draw")
+	bool bDrawVoxelCenters = false;
+
+	/** SDF volume */
+	UPROPERTY(meta = (DataflowOutput))
+	FDataflowFloatVolume FloatVolume;
+
+	virtual void Evaluate(UE::Dataflow::FContext& Context, const FDataflowOutput* Out) const override;
+
+#if WITH_EDITOR
+	virtual bool CanDebugDraw() const override
+	{
+		return true;
+	}
+	virtual bool CanDebugDrawViewMode(const FName& ViewModeName) const override;
+	virtual void DebugDraw(UE::Dataflow::FContext& Context, IDataflowDebugDrawInterface& DataflowRenderingInterface, const FDebugDrawParameters& DebugDrawParameters) const override;
+#endif
+
+public:
+	FMakeSDFPlatonicSolidDataflowNode(const UE::Dataflow::FNodeParameters& InParam, FGuid InGuid = FGuid::NewGuid());
+};
+
+/**
+ * Outputs info of a volume
+ * Grid(s), resolution, type, voxel count, voxel size
+ */
+USTRUCT(meta = (Experimental, DataflowVolume))
+struct FVolumeInfoDataflowNode : public FDataflowNode
+{
+	GENERATED_USTRUCT_BODY()
+	DATAFLOW_NODE_DEFINE_INTERNAL(FVolumeInfoDataflowNode, "VolumeInfo", "Volume|Utilities", "")
+
+private:
+	/** Volume */
+	UPROPERTY(meta = (DataflowInput))
+	FDataflowFloatVolume FloatVolume;
+
+	/** Formatted string containing the groups and attributes */
+	UPROPERTY(meta = (DataflowOutput))
+	FString String;
+
+	virtual void Evaluate(UE::Dataflow::FContext& Context, const FDataflowOutput* Out) const override;
+
+public:
+	FVolumeInfoDataflowNode(const UE::Dataflow::FNodeParameters& InParam, FGuid InGuid = FGuid::NewGuid());
+};
+
+/**
+ *
+ * Convert a Collection to a volume
+ *
+ */
+USTRUCT(meta = (Experimental, DataflowVolume))
+struct FCollectionToVolumeDataflowNode : public FDataflowNode
+{
+	GENERATED_USTRUCT_BODY()
+	DATAFLOW_NODE_DEFINE_INTERNAL(FCollectionToVolumeDataflowNode, "CollectionToVolume", "Volume|Utilities", "")
+
+private:
+	/** Collection to convert to volume */
+	UPROPERTY(meta = (DataflowInput, DataflowIntrinsic));
+	FManagedArrayCollection Collection;
+
+	/** 
+	 * Optional reference volume, to use this grid's orientation and voxel size to create the new one.
+	 * If the reference VDB is a Level set then match the narrow band width as well.
+	 */
+	UPROPERTY(meta = (DataflowInput));
+	FDataflowVolumeTypes ReferenceVolume;
+
+	/** The size of voxels in the SDF */
+	UPROPERTY(EditAnywhere, Category = "SDF", meta = (DataflowInput, ClampMin = 0.1f));
+	float VoxelSize = 1.f;
+
+	/** Name of the SDF grid in the VDB */
+	UPROPERTY(EditAnywhere, Category = "SDF");
+	EDataflowVolumeOutputType OutputType = EDataflowVolumeOutputType::SDF;
+
+	/** Name of the SDF grid in the VDB */
+	UPROPERTY(EditAnywhere, Category = "SDF");
+	FString GridName = FString(TEXT("LevelSet"));
+
+	/** Use world space units for narrow band */
+	UPROPERTY(EditAnywhere, Category = "SDF");
+	bool bUseWorldSpaceUnits = false;
+
+	/** */
+	UPROPERTY(EditAnywhere, Category = "SDF", meta = (ClampMin = 0.00001f, ClampMax = 10.f, EditCondition = "bUseWorldSpaceUnits == true", EditConditionHides));
+	float ExteriorBand = 1.f;
+
+	/** */
+	UPROPERTY(EditAnywhere, Category = "SDF", meta = (ClampMin = 0.00001f, ClampMax = 10.f, EditCondition = "bUseWorldSpaceUnits == true && bFillInterior == false", EditConditionHides));
+	float InteriorBand = 1.f;
+
+	/** */
+	UPROPERTY(EditAnywhere, Category = "SDF", meta = (ClampMin = 1, ClampMax = 10, EditCondition = "bUseWorldSpaceUnits == false", EditConditionHides));
+	int32 ExteriorBandVoxels = 3;
+
+	/** */
+	UPROPERTY(EditAnywhere, Category = "SDF", meta = (ClampMin = 1, ClampMax = 10, EditCondition = "bUseWorldSpaceUnits == false && bFillInterior == false", EditConditionHides));
+	int32 InteriorBandVoxels = 3;
+
+	/** */
+	UPROPERTY(EditAnywhere, Category = "SDF");
+	bool bFillInterior = false;
+
+	/**  */
+	UPROPERTY(EditAnywhere, Category = "SDF");
+	bool bPreserveHoles = false;
+
+	/** IsoValue specifying the surface for ClosestSurfacePointVolume */
+	UPROPERTY(EditAnywhere, Category = "Closest Surface Point", meta = (DataflowInput));
+	float IsoValue = 0.f;
+
+	UPROPERTY(EditAnywhere, Category = "Debug Draw", meta = (ShowOnlyInnerProperties))
+	FDataflowNodeDebugDrawSettings DebugDrawRenderSettings;
+
+	UPROPERTY(EditAnywhere, Category = "Debug Draw")
+	bool bDrawVoxels = false;
+
+	UPROPERTY(EditAnywhere, Category = "Debug Draw")
+	bool bDrawVoxelCenters = true;
+
+	/** Float volume */
+	UPROPERTY(meta = (DataflowOutput))
+	FDataflowFloatVolume FloatVolume;
+
+	/** Every voxel stores the closest face index */
+	UPROPERTY(meta = (DataflowOutput, DisplayName = "FaceIndex Volume"))
+	FDataflowIntVolume FaceIndexVolume;
+
+	virtual void Evaluate(UE::Dataflow::FContext& Context, const FDataflowOutput* Out) const override;
+
+#if WITH_EDITOR
+	virtual bool CanDebugDraw() const override
+	{
+		return true;
+	}
+	virtual bool CanDebugDrawViewMode(const FName& ViewModeName) const override;
+	virtual void DebugDraw(UE::Dataflow::FContext& Context, IDataflowDebugDrawInterface& DataflowRenderingInterface, const FDebugDrawParameters& DebugDrawParameters) const override;
+#endif
+
+public:
+	FCollectionToVolumeDataflowNode(const UE::Dataflow::FNodeParameters& InParam, FGuid InGuid = FGuid::NewGuid());
+};
+
+/**
+ *
+ * Fill an SDF/Fog volume with spheres
+ *
+ */
+USTRUCT(meta = (Experimental, DataflowVolume))
+struct FVolumeToSpheresDataflowNode : public FDataflowNode
+{
+	GENERATED_USTRUCT_BODY()
+	DATAFLOW_NODE_DEFINE_INTERNAL(FVolumeToSpheresDataflowNode, "VolumeToSpheres", "Volume|Utilities", "")
+
+private:
+	/** Float volume */
+	UPROPERTY(meta = (DataflowInput))
+	FDataflowFloatVolume FloatVolume;
+
+	UPROPERTY(EditAnywhere, Category = "Spheres", meta = (DataflowInput, DisplayName = "Minimum Sphere Count"));
+	int32 MinSphereCount = 1;
+
+	UPROPERTY(EditAnywhere, Category = "Spheres", meta = (DataflowInput, DisplayName = "Maximum Sphere Count"));
+	int32 MaxSphereCount = 50;
+
+	UPROPERTY(EditAnywhere, Category = "Spheres");
+	bool bOverlapping = false;
+
+	UPROPERTY(EditAnywhere, Category = "Spheres", meta = (DataflowInput, DisplayName = "Minimum Radius", ClampMin = 0.001f));
+	float MinRadius = 1.f;
+
+	UPROPERTY(EditAnywhere, Category = "Spheres", meta = (DataflowInput, DisplayName = "Maximum Radius", ClampMin = 0.001f));
+	float MaxRadius = 1000.f;
+
+	UPROPERTY(EditAnywhere, Category = "Spheres", meta = (DataflowInput));
+	float IsoValue = 0.f;
+
+	UPROPERTY(EditAnywhere, Category = "Spheres", meta = (DataflowInput, ClampMin = 1));
+	int32 InstanceCount = 10000;
+
+	UPROPERTY(EditAnywhere, Category = "Debug Draw", meta = (ShowOnlyInnerProperties))
+	FDataflowNodeDebugDrawSettings DebugDrawRenderSettings;
+
+	UPROPERTY(EditAnywhere, Category = "Debug Draw")
+	bool bDrawSpheres = true;
+
+	UPROPERTY(EditAnywhere, Category = "Debug Draw")
+	bool bDrawSphereCenters = false;
+
+	/**  */
+	UPROPERTY(meta = (DataflowOutput))
+	FManagedArrayCollection Collection;
+
+	/**  */
+	UPROPERTY(meta = (DataflowOutput))
+	TArray<FSphere> Spheres;
+
+	virtual void Evaluate(UE::Dataflow::FContext& Context, const FDataflowOutput* Out) const override;
+
+#if WITH_EDITOR
+	virtual bool CanDebugDraw() const override
+	{
+		return true;
+	}
+	virtual bool CanDebugDrawViewMode(const FName& ViewModeName) const override;
+	virtual void DebugDraw(UE::Dataflow::FContext& Context, IDataflowDebugDrawInterface& DataflowRenderingInterface, const FDebugDrawParameters& DebugDrawParameters) const override;
+#endif
+
+public:
+	FVolumeToSpheresDataflowNode(const UE::Dataflow::FNodeParameters& InParam, FGuid InGuid = FGuid::NewGuid());
+};
+
+/**
+ * Convert a volume to Volume/Polygons
+ */
+USTRUCT(meta = (Experimental, DataflowVolume))
+struct FConvertVolumeDataflowNode : public FDataflowNode
+{
+	GENERATED_USTRUCT_BODY()
+	DATAFLOW_NODE_DEFINE_INTERNAL(FConvertVolumeDataflowNode, "ConvertVolume", "Volume|Utilities", "")
+	DATAFLOW_NODE_RENDER_TYPE("SurfaceRender", FGeometryCollection::StaticType(), "Collection")
+
+private:
+	/** Float volume */
+	UPROPERTY(meta = (DataflowInput, DataflowOutput, DataflowPassthrough = "FloatVolume", DataflowIntrinsic))
+	FDataflowFloatVolume FloatVolume;
+
+	UPROPERTY(EditAnywhere, Category = "Convert", meta = (DisplayName = "Convert To"));
+	EDataflowVolumeConvertSDFTo ConvertTo = EDataflowVolumeConvertSDFTo::Collection;
+
+	UPROPERTY(EditAnywhere, Category = "Convert", meta = (UiMin = -1, UIMax = 1, EditCondition = "ConvertTo == EDataflowVolumeConvertSDFTo::Collection", EditConditionHides));
+	float IsoValue = 0.f;
+
+	UPROPERTY(EditAnywhere, Category = "Convert", meta = (ClampMin = 0, ClampMax = 1, EditCondition = "ConvertTo == EDataflowVolumeConvertSDFTo::Collection", EditConditionHides));
+	float Adaptivity = 0.f;
+
+	UPROPERTY(EditAnywhere, Category = "Convert", meta = (DisplayName = "Class", EditCondition = "ConvertTo == EDataflowVolumeConvertSDFTo::SDF", EditConditionHides));
+	EDataflowVolumeConvertSDFGridClass GridClass = EDataflowVolumeConvertSDFGridClass::NoChange;
+
+	UPROPERTY(EditAnywhere, Category = "Convert", meta = (DisplayName = "Type", EditCondition = "ConvertTo == EDataflowVolumeConvertSDFTo::SDF && GridClass == EDataflowVolumeConvertSDFGridClass::NoChange", EditConditionHides));
+	EDataflowVolumeConvertSDFGridType GridType = EDataflowVolumeConvertSDFGridType::NoChange;
+
+	UPROPERTY(EditAnywhere, Category = "Convert", meta = (ClampMin = 0, ClampMax = 1, EditCondition = "ConvertTo == EDataflowVolumeConvertSDFTo::SDF && GridClass == EDataflowVolumeConvertSDFGridClass::FogToSDF", EditConditionHides));
+	float FogIsoValue = 0.5f;
+
+	UPROPERTY(EditAnywhere, Category = "Convert", meta = (EditCondition = "ConvertTo == EDataflowVolumeConvertSDFTo::SDF", EditConditionHides));
+	bool bPruneTolerance = true;
+
+	UPROPERTY(EditAnywhere, Category = "Convert", meta = (EditCondition = "ConvertTo == EDataflowVolumeConvertSDFTo::SDF && bPruneTolerance == true", EditConditionHides));
+	float Tolerance = 0.f;
+
+	UPROPERTY(EditAnywhere, Category = "Convert", meta = (DisplayName = "Signed-Flood Fill Output", EditCondition = "ConvertTo == EDataflowVolumeConvertSDFTo::SDF", EditConditionHides));
+	bool bFloodFillOutput = true;
+
+	UPROPERTY(EditAnywhere, Category = "Convert", meta = (DisplayName = "Activate Interior Voxels", EditCondition = "ConvertTo == EDataflowVolumeConvertSDFTo::SDF", EditConditionHides));
+	bool bActivateInteriror = true;
+
+	UPROPERTY(EditAnywhere, Category = "Debug Draw", meta = (ShowOnlyInnerProperties))
+	FDataflowNodeDebugDrawSettings DebugDrawRenderSettings;
+
+	UPROPERTY(EditAnywhere, Category = "Debug Draw")
+	bool bDrawVoxels = true;
+
+	UPROPERTY(EditAnywhere, Category = "Debug Draw")
+	bool bDrawVoxelCenters = false;
+
+	/** Collection output */
+	UPROPERTY(meta = (DataflowOutput));
+	FManagedArrayCollection Collection;
+
+	virtual void Evaluate(UE::Dataflow::FContext& Context, const FDataflowOutput* Out) const override;
+
+#if WITH_EDITOR
+	virtual bool CanDebugDraw() const override
+	{
+		return true;
+	}
+	virtual bool CanDebugDrawViewMode(const FName& ViewModeName) const override;
+	virtual void DebugDraw(UE::Dataflow::FContext& Context, IDataflowDebugDrawInterface& DataflowRenderingInterface, const FDebugDrawParameters& DebugDrawParameters) const override;
+#endif
+public:
+	FConvertVolumeDataflowNode(const UE::Dataflow::FNodeParameters& InParam, FGuid InGuid = FGuid::NewGuid());
+};
+
+/**
+ * Uniformly scatter a total amount of points in active regions of a LevelSet or Fog volume
+ */
+USTRUCT(meta = (Experimental, DataflowVolume))
+struct FUniformVolumeScatterDataflowNode : public FDataflowNode
+{
+	GENERATED_USTRUCT_BODY()
+	DATAFLOW_NODE_DEFINE_INTERNAL(FUniformVolumeScatterDataflowNode, "UniformVolumeScatter", "Volume|Generators", "")
+
+private:
+	/** Float volume */
+	UPROPERTY(meta = (DataflowInput, DataflowIntrinsic))
+	FDataflowFloatVolume FloatVolume;
+
+	/** Minimum for the random range */
+	UPROPERTY(EditAnywhere, Category = "Scatter", meta = (DataflowInput, ClampMin = 1));
+	int32 MinNumberOfPoints = 100;
+
+	/** Maximum for the random range */
+	UPROPERTY(EditAnywhere, Category = "Scatter", meta = (DataflowInput, ClampMin = 1));
+	int32 MaxNumberOfPoints = 100;
+
+	/** Seed for random */
+	UPROPERTY(EditAnywhere, Category = "Scatter", meta = (DataflowInput, ClampMin = "0"));
+	int32 RandomSeed = 0;
+
+	/** IsoValue */
+	UPROPERTY(EditAnywhere, Category = "Volume", meta = (DataflowInput));
+	float IsoValue = 0.f;
+
+	/** */
+	UPROPERTY(EditAnywhere, Category = "Volume", meta = (DataflowInput, ClampMin = 0, ClampMax = 1));
+	float Spread = 1.f;
+
+	UPROPERTY(EditAnywhere, Category = "Debug Draw")
+	FLinearColor Color = FLinearColor(0.7, 0.0, 1.0, 1.0);
+
+	/** Generated points */
+	UPROPERTY(meta = (DataflowOutput))
+	TArray<FVector> Points;
+
+	virtual void Evaluate(UE::Dataflow::FContext& Context, const FDataflowOutput* Out) const override;
+
+#if WITH_EDITOR
+	virtual bool CanDebugDraw() const override
+	{
+		return true;
+	}
+	virtual bool CanDebugDrawViewMode(const FName& ViewModeName) const override;
+	virtual void DebugDraw(UE::Dataflow::FContext& Context, IDataflowDebugDrawInterface& DataflowRenderingInterface, const FDebugDrawParameters& DebugDrawParameters) const override;
+#endif
+
+public:
+	FUniformVolumeScatterDataflowNode(const UE::Dataflow::FNodeParameters& InParam, FGuid InGuid = FGuid::NewGuid());
+};
+
+/**
+ *
+ * Uniformly scatter a fixed number of points per active voxel of a LevelSet or Fog volume. If the pointsPerVoxel
+ * value provided is a fractional value, each voxel calculates a delta value of
+ * how likely it is to contain an extra point. 
+ *
+ */
+USTRUCT(meta = (Experimental, DataflowVolume))
+struct FDenseUniformVolumeScatterDataflowNode : public FDataflowNode
+{
+	GENERATED_USTRUCT_BODY()
+	DATAFLOW_NODE_DEFINE_INTERNAL(FDenseUniformVolumeScatterDataflowNode, "DenseUniformVolumeScatter", "Volume|Generators", "")
+
+private:
+	/** Float volume */
+	UPROPERTY(meta = (DataflowInput, DataflowIntrinsic))
+	FDataflowFloatVolume FloatVolume;
+
+	/** Minimum for the random range */
+	UPROPERTY(EditAnywhere, Category = "Scatter", meta = (DataflowInput, ClampMin = 0));
+	float MinNumberOfPointsPerVoxel = 1.f;
+
+	/** Maximum for the random range */
+	UPROPERTY(EditAnywhere, Category = "Scatter", meta = (DataflowInput, ClampMin = 0));
+	float MaxNumberOfPointsPerVoxel = 1.f;
+
+	/** Seed for random */
+	UPROPERTY(EditAnywhere, Category = "Scatter", meta = (DataflowInput, ClampMin = "0"));
+	int32 RandomSeed = 0;
+
+	/** IsoValue */
+	UPROPERTY(EditAnywhere, Category = "Volume", meta = (DataflowInput));
+	float IsoValue = 0.f;
+
+	/** */
+	UPROPERTY(EditAnywhere, Category = "Volume", meta = (DataflowInput, ClampMin = 0, ClampMax = 1));
+	float Spread = 1.f;
+
+	UPROPERTY(EditAnywhere, Category = "Debug Draw")
+	FLinearColor Color = FLinearColor(0.7, 0.0, 1.0, 1.0);
+
+	/** Generated points */
+	UPROPERTY(meta = (DataflowOutput))
+	TArray<FVector> Points;
+
+	virtual void Evaluate(UE::Dataflow::FContext& Context, const FDataflowOutput* Out) const override;
+
+#if WITH_EDITOR
+	virtual bool CanDebugDraw() const override
+	{
+		return true;
+	}
+	virtual bool CanDebugDrawViewMode(const FName& ViewModeName) const override;
+	virtual void DebugDraw(UE::Dataflow::FContext& Context, IDataflowDebugDrawInterface& DataflowRenderingInterface, const FDebugDrawParameters& DebugDrawParameters) const override;
+#endif
+
+public:
+	FDenseUniformVolumeScatterDataflowNode(const UE::Dataflow::FNodeParameters& InParam, FGuid InGuid = FGuid::NewGuid());
+};
+
+/**
+ *
+ * Non uniformly scatter points per active voxel. The pointsPerVoxel value is used
+ * to weight each grids cell value to compute a fixed number of points for every
+ * active voxel. If the computed result is a fractional value, each voxel calculates
+ * a delta value of how likely it is to contain an extra point.
+ *
+ */
+USTRUCT(meta = (Experimental, DataflowVolume))
+struct FNonUniformVolumeScatterDataflowNode : public FDataflowNode
+{
+	GENERATED_USTRUCT_BODY()
+	DATAFLOW_NODE_DEFINE_INTERNAL(FNonUniformVolumeScatterDataflowNode, "NonUniformVolumeScatter", "Volume|Generators", "")
+
+private:
+	/** Float volume */
+	UPROPERTY(meta = (DataflowInput, DataflowIntrinsic))
+	FDataflowFloatVolume FloatVolume;
+
+	/** Minimum for the random range */
+	UPROPERTY(EditAnywhere, Category = "Scatter", meta = (DataflowInput, ClampMin = 0));
+	float MinNumberOfPointsPerVoxel = .1f;
+
+	/** Maximum for the random range */
+	UPROPERTY(EditAnywhere, Category = "Scatter", meta = (DataflowInput, ClampMin = 0));
+	float MaxNumberOfPointsPerVoxel = .3f;
+
+	/** Seed for random */
+	UPROPERTY(EditAnywhere, Category = "Scatter", meta = (DataflowInput, ClampMin = "0"));
+	int32 RandomSeed = 0;
+
+	/** IsoValue */
+	UPROPERTY(EditAnywhere, Category = "Volume", meta = (DataflowInput));
+	float IsoValue = 0.f;
+
+	/** */
+	UPROPERTY(EditAnywhere, Category = "Volume", meta = (DataflowInput, ClampMin = 0, ClampMax = 1));
+	float Spread = 1.f;
+
+	UPROPERTY(EditAnywhere, Category = "Debug Draw")
+	FLinearColor Color = FLinearColor(0.7, 0.0, 1.0, 1.0);
+
+	/** Generated points */
+	UPROPERTY(meta = (DataflowOutput))
+	TArray<FVector> Points;
+
+	virtual void Evaluate(UE::Dataflow::FContext& Context, const FDataflowOutput* Out) const override;
+
+#if WITH_EDITOR
+	virtual bool CanDebugDraw() const override
+	{
+		return true;
+	}
+	virtual bool CanDebugDrawViewMode(const FName& ViewModeName) const override;
+	virtual void DebugDraw(UE::Dataflow::FContext& Context, IDataflowDebugDrawInterface& DataflowRenderingInterface, const FDebugDrawParameters& DebugDrawParameters) const override;
+#endif
+
+public:
+	FNonUniformVolumeScatterDataflowNode(const UE::Dataflow::FNodeParameters& InParam, FGuid InGuid = FGuid::NewGuid());
+};
+
+/**
+ * Sample a volume at input point's location
+ */
+USTRUCT(meta = (Experimental, DataflowVolume))
+struct FVolumeSampleToPointsDataflowNode : public FDataflowNode
+{
+	GENERATED_USTRUCT_BODY()
+	DATAFLOW_NODE_DEFINE_INTERNAL(FVolumeSampleToPointsDataflowNode, "VolumeSampleToPoints", "Volume|Sample", "")
+
+private:
+	/** Volume to sample */
+	UPROPERTY(meta = (DataflowInput, DataflowIntrinsic))
+	FDataflowVolumeTypes Volume;
+
+	/** Points to sample at */ 
+	UPROPERTY(meta = (DataflowInput, DataflowIntrinsic))
+	FDataflowVectorArrayTypes SamplePoints;
+
+	UPROPERTY(EditAnywhere, Category = "Sampling");
+	EDataflowVolumeSampleType SampleType = EDataflowVolumeSampleType::Float;
+
+	UPROPERTY(EditAnywhere, Category = "Debug Draw")
+	FLinearColor Color = FLinearColor(0.2, 1.0, 1.0, 1.0);
+
+	UPROPERTY(EditAnywhere, Category = "Sampling")
+	FString Attribute = FString("Attribute");
+
+	// TArray<float>, TArray<int32>, TArray<FVector3f> 
+	UPROPERTY(meta = (DataflowOutput));
+	FDataflowAnyType OutArray;
+
+	UPROPERTY(meta = (DataflowOutput));
+	FManagedArrayCollection PointsCollection;
+
+	virtual void Evaluate(UE::Dataflow::FContext& Context, const FDataflowOutput* Out) const override;
+#if WITH_EDITOR
+	virtual bool CanDebugDraw() const override
+	{
+		return true;
+	}
+	virtual bool CanDebugDrawViewMode(const FName& ViewModeName) const override;
+	virtual void DebugDraw(UE::Dataflow::FContext& Context, IDataflowDebugDrawInterface& DataflowRenderingInterface, const FDebugDrawParameters& DebugDrawParameters) const override;
+#endif
+public:
+	FVolumeSampleToPointsDataflowNode(const UE::Dataflow::FNodeParameters& InParam, FGuid InGuid = FGuid::NewGuid());
+
+	void SetOutTypeFromSampleType();
+	virtual void OnPropertyChanged(UE::Dataflow::FContext& Context, const FPropertyChangedEvent& InPropertyChangedEvent) override;
+};
+
+/**
+ * Extracts 2d slices from a volume
+ */
+USTRUCT(meta = (Experimental, DataflowVolume))
+struct FVolumeSliceDataflowNode : public FDataflowNode
+{
+	GENERATED_USTRUCT_BODY()
+	DATAFLOW_NODE_DEFINE_INTERNAL(FVolumeSliceDataflowNode, "VolumeSlice", "Volume|Visualize", "Visualize Scalar Vector Curve 2D")
+
+private:
+	/** Volume to sample */
+	UPROPERTY(meta = (DataflowInput, DataflowOutput, DataflowPassthrough = "Volume", DataflowIntrinsic))
+//	FDataflowFloatVolume Volume;
+	FDataflowVolumeTypes Volume;
+
+	/**  */
+	UPROPERTY(EditAnywhere, Category = "Slice Settings");
+	EDataflowVolumeSliceMethod Method = EDataflowVolumeSliceMethod::Points;
+
+	/**  */
+	UPROPERTY(EditAnywhere, Category = "Slice Settings");
+	EDataflowVolumeSlicePlane Plane = EDataflowVolumeSlicePlane::XYPlane;
+
+	/**  */
+	UPROPERTY(EditAnywhere, Category = "Slice Settings", meta = (ClampMin = "1", ClampMax = "10"));
+	float BoundingBoxScale = 1.f;
+
+	/**  */
+	UPROPERTY(EditAnywhere, Category = "Slice Settings", meta = (ClampMin = "-1", ClampMax = "1"));
+	float Offset = 0.f;
+
+	UPROPERTY(EditAnywhere, Category = "Visualization Settings");
+	bool bVisualize = true;
+
+	/**  */
+	UPROPERTY(EditAnywhere, Category = "Visualization Settings - Scalar Visualization Settings", meta = (EditCondition = "bVisualize == true", EditConditionHides));
+	EDataflowVolumeSliceRamp ColorRampType = EDataflowVolumeSliceRamp::InfraRed;
+
+	UPROPERTY(EditAnywhere, Category = "Visualization Settings - Scalar Visualization Settings", meta = (EditCondition = "bVisualize == true", EditConditionHides));
+	FDataflowColorRamp ColorRamp;
+
+	UPROPERTY(EditAnywhere, Category = "Visualization Settings - Scalar Visualization Settings");
+	bool bInvertColorRamp = false;
+
+	/**  */
+	UPROPERTY(EditAnywhere, Category = "Visualization Settings - Scalar Visualization Settings", meta = (ClampMin = "-100", ClampMax = "100", EditCondition = "bVisualize == true", EditConditionHides));
+	FVector2f Range = FVector2f(0.f, 1.f);
+
+	UPROPERTY(EditAnywhere, Category = "Visualization Settings - Vector Visualization Settings", meta = (ClampMin = "0.01", ClampMax = "100", EditCondition = "bVisualize == true", EditConditionHides));
+	float VectorScale = 1.f;
+
+	UPROPERTY(EditAnywhere, Category = "Visualization Settings - Vector Visualization Settings", meta = (ClampMin = "0.01", ClampMax = "100", EditCondition = "bVisualize == true", EditConditionHides));
+	float PointScale = 4.f;
+
+	UPROPERTY(EditAnywhere, Category = "Visualization Settings - Vector Visualization Settings", meta = (ClampMin = "0.01", ClampMax = "100", EditCondition = "bVisualize == true", EditConditionHides));
+	float LineWidth = 1.f;
+
+	UPROPERTY(EditAnywhere, Category = "Visualization Settings - Vector Visualization Settings", meta = (ClampMin = "0.0001", ClampMax = "100", EditCondition = "bVisualize == true", EditConditionHides));
+	float MinLength = 0.1f;
+
+	UPROPERTY(EditAnywhere, Category = "Visualization Settings - Vector Visualization Settings", meta = (EditCondition = "bVisualize == true", EditConditionHides));
+	FLinearColor PointColor = FLinearColor::Yellow;
+
+	UPROPERTY(EditAnywhere, Category = "Visualization Settings - Vector Visualization Settings", meta = (EditCondition = "bVisualize == true", EditConditionHides));
+	FLinearColor VectorColor = FLinearColor::Yellow;
+
+	virtual void Evaluate(UE::Dataflow::FContext& Context, const FDataflowOutput* Out) const override;
+	virtual void OnPropertyChanged(UE::Dataflow::FContext& Context, const FPropertyChangedEvent& InPropertyChangedEvent) override;
+
+#if WITH_EDITOR
+	virtual bool CanDebugDraw() const override
+	{
+		return true;
+	}
+	virtual bool CanDebugDrawViewMode(const FName& ViewModeName) const override;
+	virtual void DebugDraw(UE::Dataflow::FContext& Context, IDataflowDebugDrawInterface& DataflowRenderingInterface, const FDebugDrawParameters& DebugDrawParameters) const override;
+#endif
+
+	void OnColorCurveChanged(TArray<FRichCurve*> Curves);
+public:
+	FVolumeSliceDataflowNode(const UE::Dataflow::FNodeParameters& InParam, FGuid InGuid = FGuid::NewGuid());
+};
+
+/**
+ * 
+ */
+USTRUCT(meta = (Experimental, DataflowVolume))
+struct FVolumeAnalysisDataflowNode : public FDataflowNode
+{
+	GENERATED_USTRUCT_BODY()
+	DATAFLOW_NODE_DEFINE_INTERNAL(FVolumeAnalysisDataflowNode, "VolumeAnalysis", "Volume|Utilities", "")
+
+private:
+	/** Volume to process (scalar or vector volume) */
+	UPROPERTY(meta = (DataflowInput, DataflowIntrinsic))
+	FDataflowVolumeTypes Volume;
+
+	/**  */
+	UPROPERTY(EditAnywhere, Category = "Settings");
+	EDataflowVolumeAnalysisOperator Operator = EDataflowVolumeAnalysisOperator::Gradient;
+
+	/** Mask volume */
+	UPROPERTY(meta = (DataflowInput, DataflowOutput, DataflowPassthrough = "MaskVolume"));
+	FDataflowBoolVolume MaskVolume;
+
+	/**  */
+	UPROPERTY(EditAnywhere, Category = "Settings");
+	EDataflowVolumeAnalysisOutputName OutputName = EDataflowVolumeAnalysisOutputName::CustomName;
+
+	/**  */
+	UPROPERTY(EditAnywhere, Category = "Settings", meta = (EditCondition = "OutputName == EDataflowVolumeAnalysisOutputName::CustomName", EditConditionHides));
+	FString CustomName = FString("Custom");
+
+	/** Volume output (scalar or vector volume) */
+	UPROPERTY(meta = (DataflowOutput))
+	FDataflowVolumeTypes OutVolume;
+
+	virtual void Evaluate(UE::Dataflow::FContext& Context, const FDataflowOutput* Out) const override;
+
+public:
+	FVolumeAnalysisDataflowNode(const UE::Dataflow::FNodeParameters& InParam, FGuid InGuid = FGuid::NewGuid());
+
+	void SetOutTypeFromOperator();
+	virtual void OnPropertyChanged(UE::Dataflow::FContext& Context, const FPropertyChangedEvent& InPropertyChangedEvent) override;
+};
+
+/**
+ *
+ *
+ *
+ */
+USTRUCT(meta = (DataflowVolume))
+struct FVolumeCombineDataflowNode : public FDataflowNode
+{
+	GENERATED_USTRUCT_BODY()
+	DATAFLOW_NODE_DEFINE_INTERNAL(FVolumeCombineDataflowNode, "VolumeCombine", "Volume|Utilities", "")
+
+private:
+	/** Float A volume */
+	UPROPERTY(meta = (DataflowInput, DataflowIntrinsic))
+	FDataflowFloatVolume FloatVolumeA;
+
+	/** Float B volume */
+	UPROPERTY(meta = (DataflowInput, DataflowIntrinsic))
+	FDataflowFloatVolume FloatVolumeB;
+
+	UPROPERTY(EditAnywhere, Category = "Combine");
+	EDataflowVolumeSDFCombineOperation Operation = EDataflowVolumeSDFCombineOperation::SDFUnion;
+
+	UPROPERTY(EditAnywhere, Category = "Combine");
+	bool bFlipInputs = false;
+
+	UPROPERTY(EditAnywhere, Category = "Combine", meta = (DataflowInput, DisplayName = "A Multiplier", UIMin = -10.f, UIMax = 10.f));
+	float MultiplierA = 1.f;
+
+	UPROPERTY(EditAnywhere, Category = "Combine", meta = (DataflowInput, DisplayName = "B Multiplier", UIMin = -10.f, UIMax = 10.f));
+	float MultiplierB = 1.f;
+
+	UPROPERTY(EditAnywhere, Category = "Combine");
+	EDataflowVolumeSDFCombineResample Resample = EDataflowVolumeSDFCombineResample::BMatchA;
+
+	UPROPERTY(EditAnywhere, Category = "Combine");
+	EDataflowVolumeSDFCombineInterpolation Interpolation = EDataflowVolumeSDFCombineInterpolation::Linear;
+
+	UPROPERTY(EditAnywhere, Category = "Combine", meta = (DisplayName = "Prune Degenarate Tiles"));
+	bool bPruneDegenerateTiles = false;
+
+	UPROPERTY(EditAnywhere, Category = "Combine", meta = (DisplayName = "Prune Tolerance", UIMin = 0.f, UIMax = 1.f));
+	float PruneTol = 0.f;
+
+	UPROPERTY(EditAnywhere, Category = "Debug Draw", meta = (ShowOnlyInnerProperties))
+	FDataflowNodeDebugDrawSettings DebugDrawRenderSettings;
+
+	UPROPERTY(EditAnywhere, Category = "Debug Draw")
+	bool bDrawVoxels = true;
+
+	UPROPERTY(EditAnywhere, Category = "Debug Draw")
+	bool bDrawVoxelCenters = false;
+
+	/** Output Float volume */
+	UPROPERTY(meta = (DataflowOutput))
+	FDataflowFloatVolume FloatVolume;
+
+	virtual void Evaluate(UE::Dataflow::FContext& Context, const FDataflowOutput* Out) const override;
+#if WITH_EDITOR
+	virtual bool CanDebugDraw() const override
+	{
+		return true;
+	}
+	virtual bool CanDebugDrawViewMode(const FName& ViewModeName) const override;
+	virtual void DebugDraw(UE::Dataflow::FContext& Context, IDataflowDebugDrawInterface& DataflowRenderingInterface, const FDebugDrawParameters& DebugDrawParameters) const override;
+#endif
+
+public:
+	FVolumeCombineDataflowNode(const UE::Dataflow::FNodeParameters& InParam, FGuid InGuid = FGuid::NewGuid());
+};
+
+namespace UE::Dataflow
+{
+	void RegisterVolumeNodes();
+}
+
diff --git a/Engine/Plugins/Experimental/DataflowVolumetricPlugin/Source/DataflowVolumeNodes/Private/Dataflow/DataflowVolumeNodesPlugin.cpp b/Engine/Plugins/Experimental/DataflowVolumetricPlugin/Source/DataflowVolumeNodes/Private/Dataflow/DataflowVolumeNodesPlugin.cpp
new file mode 100644
--- /dev/null
+++ b/Engine/Plugins/Experimental/DataflowVolumetricPlugin/Source/DataflowVolumeNodes/Private/Dataflow/DataflowVolumeNodesPlugin.cpp
@@ -0,0 +1,30 @@
+// Copyright Epic Games, Inc. All Rights Reserved.
+
+#include "Dataflow/DataflowVolumeNodesPlugin.h"
+
+#include "CoreMinimal.h"
+
+#include "Dataflow/DataflowNodeFactory.h"
+#include "Dataflow/DataflowVolumeNodes.h"
+#include "Dataflow/DataflowVolumeNodeAndPinTypeColors.h"
+#include "Dataflow/DataflowVolumeAnyType.h"
+#include "Dataflow/DataflowColorRampCustomization.h"
+
+void IDataflowVolumetricPlugin::StartupModule()
+{
+	// anytypes need to be at the top as the connection registration in the nodes are using it
+	UE::Dataflow::RegisterVolumeAnyTypes();
+
+	UE::Dataflow::RegisterVolumeNodes();
+
+	UE::Dataflow::RegisterVolumeColors();
+
+	FModuleManager::Get().LoadModule("DataflowVolumeCore");
+}
+
+void IDataflowVolumetricPlugin::ShutdownModule()
+{
+}
+
+IMPLEMENT_MODULE(IDataflowVolumetricPlugin, DataflowVolumeNodes)
+
diff --git a/Engine/Plugins/Experimental/DataflowVolumetricPlugin/Source/DataflowVolumeNodes/Private/Dataflow/DataflowVolumeNodesPlugin.h b/Engine/Plugins/Experimental/DataflowVolumetricPlugin/Source/DataflowVolumeNodes/Private/Dataflow/DataflowVolumeNodesPlugin.h
new file mode 100644
--- /dev/null
+++ b/Engine/Plugins/Experimental/DataflowVolumetricPlugin/Source/DataflowVolumeNodes/Private/Dataflow/DataflowVolumeNodesPlugin.h
@@ -0,0 +1,18 @@
+// Copyright Epic Games, Inc. All Rights Reserved.
+
+#pragma once
+
+#include "CoreMinimal.h"
+#include "Modules/ModuleInterface.h"
+#include "Modules/ModuleManager.h"
+
+/**
+ * The public interface to this module
+ */
+class IDataflowVolumetricPlugin : public IModuleInterface
+{
+public:
+	virtual void StartupModule();
+	virtual void ShutdownModule();
+};
+
