Dataflow:
- Added VolumeTextureTerminal node

diff --git a/Engine/Plugins/Experimental/DataflowVolumetricPlugin/Source/DataflowVolumeCore/Internal/Dataflow/DataflowBoolVolume.h b/Engine/Plugins/Experimental/DataflowVolumetricPlugin/Source/DataflowVolumeCore/Internal/Dataflow/DataflowBoolVolume.h
--- a/Engine/Plugins/Experimental/DataflowVolumetricPlugin/Source/DataflowVolumeCore/Internal/Dataflow/DataflowBoolVolume.h
+++ b/Engine/Plugins/Experimental/DataflowVolumetricPlugin/Source/DataflowVolumeCore/Internal/Dataflow/DataflowBoolVolume.h
@@ -29,8 +29,7 @@ public:
 	TSharedPtr<const UE::DataflowVolume::Private::FDataflowBoolVolumeImpl> GetBoolVolumeImpl() const;
 	TSharedPtr<UE::DataflowVolume::Private::FDataflowBoolVolumeImpl> GetBoolVolumeImpl();
 
-	FVector GetVoxelSize() const;
-	FBox GetVolumeBoundingBox() const;
+	DATAFLOWVOLUMECORE_API FVector GetVoxelSize() const;
 
 private:
 	FDataflowBoolVolume(TSharedPtr<UE::DataflowVolume::Private::FDataflowBoolVolumeImpl> InBoolVolume);
diff --git a/Engine/Plugins/Experimental/DataflowVolumetricPlugin/Source/DataflowVolumeCore/Internal/Dataflow/DataflowFloatVectorVolume.h b/Engine/Plugins/Experimental/DataflowVolumetricPlugin/Source/DataflowVolumeCore/Internal/Dataflow/DataflowFloatVectorVolume.h
--- a/Engine/Plugins/Experimental/DataflowVolumetricPlugin/Source/DataflowVolumeCore/Internal/Dataflow/DataflowFloatVectorVolume.h
+++ b/Engine/Plugins/Experimental/DataflowVolumetricPlugin/Source/DataflowVolumeCore/Internal/Dataflow/DataflowFloatVectorVolume.h
@@ -36,7 +36,6 @@ public:
 	TSharedPtr<UE::DataflowVolume::Private::FDataflowFloatVectorVolumeImpl> GetFloatVectorVolumeImpl();
 
 	DATAFLOWVOLUMECORE_API FVector GetVoxelSize() const;
-	DATAFLOWVOLUMECORE_API FBox GetVolumeBoundingBox() const;
 	DATAFLOWVOLUMECORE_API FString GetVectorType() const;
 
 	DATAFLOWVOLUMECORE_API void VolumeSample(const TArray<FVector>& InPoints, TArray<FVector>& OutValues) const;
diff --git a/Engine/Plugins/Experimental/DataflowVolumetricPlugin/Source/DataflowVolumeCore/Internal/Dataflow/DataflowFloatVolume.h b/Engine/Plugins/Experimental/DataflowVolumetricPlugin/Source/DataflowVolumeCore/Internal/Dataflow/DataflowFloatVolume.h
--- a/Engine/Plugins/Experimental/DataflowVolumetricPlugin/Source/DataflowVolumeCore/Internal/Dataflow/DataflowFloatVolume.h
+++ b/Engine/Plugins/Experimental/DataflowVolumetricPlugin/Source/DataflowVolumeCore/Internal/Dataflow/DataflowFloatVolume.h
@@ -35,7 +35,6 @@ public:
 	TSharedPtr<UE::DataflowVolume::Private::FDataflowFloatVolumeImpl> GetFloatVolumeImpl();
 
 	DATAFLOWVOLUMECORE_API FVector GetVoxelSize() const;
-	DATAFLOWVOLUMECORE_API FBox GetVolumeBoundingBox() const;
 
 	// Create a sphere SDF
 	DATAFLOWVOLUMECORE_API static FDataflowFloatVolume CreateSphereSDF(const float InVoxelSize, const float InRadius, const FVector& InCenter);
diff --git a/Engine/Plugins/Experimental/DataflowVolumetricPlugin/Source/DataflowVolumeCore/Internal/Dataflow/DataflowIntVolume.h b/Engine/Plugins/Experimental/DataflowVolumetricPlugin/Source/DataflowVolumeCore/Internal/Dataflow/DataflowIntVolume.h
--- a/Engine/Plugins/Experimental/DataflowVolumetricPlugin/Source/DataflowVolumeCore/Internal/Dataflow/DataflowIntVolume.h
+++ b/Engine/Plugins/Experimental/DataflowVolumetricPlugin/Source/DataflowVolumeCore/Internal/Dataflow/DataflowIntVolume.h
@@ -26,11 +26,10 @@ public:
 	FDataflowIntVolume() {}
 	virtual ~FDataflowIntVolume() {}
 
-	TSharedPtr<const UE::DataflowVolume::Private::FDataflowIntVolumeImpl> GetIntVolume() const;
-	TSharedPtr<UE::DataflowVolume::Private::FDataflowIntVolumeImpl> GetIntVolume();
+	TSharedPtr<const UE::DataflowVolume::Private::FDataflowIntVolumeImpl> GetIntVolumeImpl() const;
+	TSharedPtr<UE::DataflowVolume::Private::FDataflowIntVolumeImpl> GetIntVolumeImpl();
 
 	DATAFLOWVOLUMECORE_API FVector GetVoxelSize() const;
-	DATAFLOWVOLUMECORE_API FBox GetVolumeBoundingBox() const;
 
 	/** Sample a volume */
 	DATAFLOWVOLUMECORE_API void VolumeSample(const TArray<FVector>& InPoints, TArray<int32>& OutValues) const;
diff --git a/Engine/Plugins/Experimental/DataflowVolumetricPlugin/Source/DataflowVolumeCore/Internal/Dataflow/DataflowVolume.h b/Engine/Plugins/Experimental/DataflowVolumetricPlugin/Source/DataflowVolumeCore/Internal/Dataflow/DataflowVolume.h
--- a/Engine/Plugins/Experimental/DataflowVolumetricPlugin/Source/DataflowVolumeCore/Internal/Dataflow/DataflowVolume.h
+++ b/Engine/Plugins/Experimental/DataflowVolumetricPlugin/Source/DataflowVolumeCore/Internal/Dataflow/DataflowVolume.h
@@ -9,6 +9,8 @@ namespace UE::DataflowVolume::Private
 	class FDataflowVolumeImpl;
 }
 
+class UVolumeTexture;
+
 /* --------------------------------------------------------------------------------------------------------------- */
 /* FDataflowVolume */
 /* --------------------------------------------------------------------------------------------------------------- */
@@ -21,7 +23,11 @@ public:
 	FDataflowVolume() {}
 	virtual ~FDataflowVolume() {}
 
+	DATAFLOWVOLUMECORE_API int32 GetNumActiveVoxels() const;
 	DATAFLOWVOLUMECORE_API void GetActiveVoxels(TArray<FBox>& OutActiveVoxels) const;
+	DATAFLOWVOLUMECORE_API FBox GetVolumeBoundingBox() const;
+
+	DATAFLOWVOLUMECORE_API bool CreateVolumeTexture(UVolumeTexture* InVolumeTexture) const;
 
 	/** Prints Volume info into string */
 	DATAFLOWVOLUMECORE_API FString VolumeInfo() const;
diff --git a/Engine/Plugins/Experimental/DataflowVolumetricPlugin/Source/DataflowVolumeCore/Private/Dataflow/DataflowBoolVolume.cpp b/Engine/Plugins/Experimental/DataflowVolumetricPlugin/Source/DataflowVolumeCore/Private/Dataflow/DataflowBoolVolume.cpp
--- a/Engine/Plugins/Experimental/DataflowVolumetricPlugin/Source/DataflowVolumeCore/Private/Dataflow/DataflowBoolVolume.cpp
+++ b/Engine/Plugins/Experimental/DataflowVolumetricPlugin/Source/DataflowVolumeCore/Private/Dataflow/DataflowBoolVolume.cpp
@@ -13,16 +13,22 @@ FDataflowBoolVolume::FDataflowBoolVolume(TSharedPtr<UE::DataflowVolume::Private:
 	static_assert(sizeof(FDataflowVolume) == sizeof(*this));
 }
 
+/* --------------------------------------------------------------------------------------------------------------- */
+
 TSharedPtr<const UE::DataflowVolume::Private::FDataflowBoolVolumeImpl> FDataflowBoolVolume::GetBoolVolumeImpl() const
 {
 	return StaticCastSharedPtr<const UE::DataflowVolume::Private::FDataflowBoolVolumeImpl>(Volume);
 }
 
+/* --------------------------------------------------------------------------------------------------------------- */
+
 TSharedPtr<UE::DataflowVolume::Private::FDataflowBoolVolumeImpl> FDataflowBoolVolume::GetBoolVolumeImpl()
 {
 	return StaticCastSharedPtr<UE::DataflowVolume::Private::FDataflowBoolVolumeImpl>(Volume);
 }
 
+/* --------------------------------------------------------------------------------------------------------------- */
+
 FVector FDataflowBoolVolume::GetVoxelSize() const
 {
 	using namespace UE::DataflowVolume::Private;
@@ -35,15 +41,3 @@ FVector FDataflowBoolVolume::GetVoxelSize() const
 	return {};
 }
 
-FBox FDataflowBoolVolume::GetVolumeBoundingBox() const
-{
-	using namespace UE::DataflowVolume::Private;
-
-	if (TSharedPtr<const FDataflowBoolVolumeImpl> BoolVolumeImpl = GetBoolVolumeImpl())
-	{
-		return BoolVolumeImpl->GetActiveVoxelBoundingBox();
-	}
-
-	return {};
-}
-
diff --git a/Engine/Plugins/Experimental/DataflowVolumetricPlugin/Source/DataflowVolumeCore/Private/Dataflow/DataflowBoolVolumeImpl.cpp b/Engine/Plugins/Experimental/DataflowVolumetricPlugin/Source/DataflowVolumeCore/Private/Dataflow/DataflowBoolVolumeImpl.cpp
--- a/Engine/Plugins/Experimental/DataflowVolumetricPlugin/Source/DataflowVolumeCore/Private/Dataflow/DataflowBoolVolumeImpl.cpp
+++ b/Engine/Plugins/Experimental/DataflowVolumetricPlugin/Source/DataflowVolumeCore/Private/Dataflow/DataflowBoolVolumeImpl.cpp
@@ -1,6 +1,8 @@
 // Copyright Epic Games, Inc. All Rights Reserved.
 
 #include "Dataflow/DataflowBoolVolumeImpl.h"
+#include <Chaos/ArrayND.h>
+#include "Engine/VolumeTexture.h"
 
 namespace UE::DataflowVolume::Private
 {
@@ -14,11 +16,15 @@ namespace UE::DataflowVolume::Private
 		BoolGrid = InBoolGrid;
 	}
 
+	/* --------------------------------------------------------------------------------------------------------------- */
+
 	openvdb::BoolGrid::Ptr FDataflowBoolVolumeImpl::GetGrid() const
 	{
 		return BoolGrid;
 	}
 
+	/* --------------------------------------------------------------------------------------------------------------- */
+
 	FString FDataflowBoolVolumeImpl::GetGridType() const
 	{
 		if (BoolGrid)
@@ -30,6 +36,8 @@ namespace UE::DataflowVolume::Private
 		return {};
 	}
 
+	/* --------------------------------------------------------------------------------------------------------------- */
+
 	FString FDataflowBoolVolumeImpl::VolumeInfo() const
 	{
 		FString BaseString = FDataflowVolumeImpl::VolumeInfo();
@@ -48,6 +56,8 @@ namespace UE::DataflowVolume::Private
 		return OutputStr;
 	}
 
+	/* --------------------------------------------------------------------------------------------------------------- */
+
 	void FDataflowBoolVolumeImpl::GetActiveVoxels(TArray<FBox>& OutActiveVoxels) const
 	{
 		if (BoolGrid)
@@ -79,5 +89,61 @@ namespace UE::DataflowVolume::Private
 			}
 		}
 	}
+
+	/* --------------------------------------------------------------------------------------------------------------- */
+
+	FBox FDataflowBoolVolumeImpl::GetVolumeBoundingBox() const
+	{
+		return this->GetActiveVoxelBoundingBox();
+	}
+
+	/* --------------------------------------------------------------------------------------------------------------- */
+
+	bool FDataflowBoolVolumeImpl::CreateVolumeTexture(UVolumeTexture* InVolumeTexture) const
+	{
+		if (InVolumeTexture == nullptr)
+		{
+			return false;
+		}
+
+		const EPixelFormat PixelFormat = EPixelFormat::PF_R16F;
+
+		if (!BoolGrid->empty())
+		{
+			openvdb::Coord ActiveVoxelDim = BoolGrid->evalActiveVoxelDim();
+			openvdb::CoordBBox Bbox = BoolGrid->evalActiveVoxelBoundingBox();
+			const int32 XMin = Bbox.min().x();
+			const int32 YMin = Bbox.min().y();
+			const int32 ZMin = Bbox.min().z();
+
+			openvdb::BoolGrid::Accessor Accessor = BoolGrid->getAccessor();
+			
+			auto QueryVoxel = [&Accessor, &XMin, &YMin, &ZMin](int32 PosX, int32 PosY, int32 PosZ, void* Value)
+			{
+				openvdb::Coord Coord(PosX + XMin, PosY + YMin, PosZ + ZMin);
+				FFloat16* const Voxel = static_cast<FFloat16*>(Value);
+
+				if (Accessor.isValueOn(Coord))
+				{
+					Voxel[0].Set((float)Accessor.getValue(Coord));
+				}
+				else
+				{
+					Voxel[0].Set(0.f);
+				}
+			};
+
+			const bool bSuccess = InVolumeTexture->UpdateSourceFromFunction(QueryVoxel, ActiveVoxelDim[0], ActiveVoxelDim[1], ActiveVoxelDim[2], TSF_R16F);
+
+			if (!bSuccess)
+			{
+				return false;
+			}
+
+			return true;
+		}
+
+		return false;
+	}
 }
 
diff --git a/Engine/Plugins/Experimental/DataflowVolumetricPlugin/Source/DataflowVolumeCore/Private/Dataflow/DataflowBoolVolumeImpl.h b/Engine/Plugins/Experimental/DataflowVolumetricPlugin/Source/DataflowVolumeCore/Private/Dataflow/DataflowBoolVolumeImpl.h
--- a/Engine/Plugins/Experimental/DataflowVolumetricPlugin/Source/DataflowVolumeCore/Private/Dataflow/DataflowBoolVolumeImpl.h
+++ b/Engine/Plugins/Experimental/DataflowVolumetricPlugin/Source/DataflowVolumeCore/Private/Dataflow/DataflowBoolVolumeImpl.h
@@ -16,11 +16,16 @@ namespace UE::DataflowVolume::Private
 		openvdb::BoolGrid::Ptr BoolGrid = nullptr;
 
 	public:
-		void SetGrid(const openvdb::BoolGrid::Ptr& InBoolGrid, const FString& InGridName = "");
+		// Virtual functions
+		virtual void GetActiveVoxels(TArray<FBox>& OutActiveVoxels) const override;
+		virtual FBox GetVolumeBoundingBox() const override;
+		virtual bool CreateVolumeTexture(UVolumeTexture* InVolumeTexture) const override;
+
 		openvdb::BoolGrid::Ptr GetGrid() const;
+		void SetGrid(const openvdb::BoolGrid::Ptr& InBoolGrid, const FString& InGridName = "");
 		FString GetGridType() const;
 		FString VolumeInfo() const;
-		virtual void GetActiveVoxels(TArray<FBox>& OutActiveVoxels) const override;
+
 	};
 }
 
diff --git a/Engine/Plugins/Experimental/DataflowVolumetricPlugin/Source/DataflowVolumeCore/Private/Dataflow/DataflowFloatVectorVolume.cpp b/Engine/Plugins/Experimental/DataflowVolumetricPlugin/Source/DataflowVolumeCore/Private/Dataflow/DataflowFloatVectorVolume.cpp
--- a/Engine/Plugins/Experimental/DataflowVolumetricPlugin/Source/DataflowVolumeCore/Private/Dataflow/DataflowFloatVectorVolume.cpp
+++ b/Engine/Plugins/Experimental/DataflowVolumetricPlugin/Source/DataflowVolumeCore/Private/Dataflow/DataflowFloatVectorVolume.cpp
@@ -66,20 +66,6 @@ FVector FDataflowFloatVectorVolume::GetVoxelSize() const
 
 /* --------------------------------------------------------------------------------------------------------------- */
 
-FBox FDataflowFloatVectorVolume::GetVolumeBoundingBox() const
-{
-	using namespace UE::DataflowVolume::Private;
-
-	if (TSharedPtr<const FDataflowFloatVectorVolumeImpl> FloatVectorVolumeImpl = GetFloatVectorVolumeImpl())
-	{
-		return FloatVectorVolumeImpl->GetActiveVoxelBoundingBox();
-	}
-
-	return {};
-}
-
-/* --------------------------------------------------------------------------------------------------------------- */
-
 FDataflowFloatVolume FDataflowFloatVectorVolume::ComputeDivergence(const EDataflowVolumeAnalysisOutputName OutputName, const FString& CustomName, const FDataflowBoolVolume* MaskVolume) const
 {
 	using namespace UE::DataflowVolume::Private;
@@ -151,6 +137,3 @@ FDataflowFloatVectorVolume FDataflowFloatVectorVolume::ComputeNormalize(const ED
 	return FDataflowFloatVectorVolume();
 }
 
-/* --------------------------------------------------------------------------------------------------------------- */
-
-
diff --git a/Engine/Plugins/Experimental/DataflowVolumetricPlugin/Source/DataflowVolumeCore/Private/Dataflow/DataflowFloatVectorVolumeImpl.cpp b/Engine/Plugins/Experimental/DataflowVolumetricPlugin/Source/DataflowVolumeCore/Private/Dataflow/DataflowFloatVectorVolumeImpl.cpp
--- a/Engine/Plugins/Experimental/DataflowVolumetricPlugin/Source/DataflowVolumeCore/Private/Dataflow/DataflowFloatVectorVolumeImpl.cpp
+++ b/Engine/Plugins/Experimental/DataflowVolumetricPlugin/Source/DataflowVolumeCore/Private/Dataflow/DataflowFloatVectorVolumeImpl.cpp
@@ -3,6 +3,7 @@
 #include "Dataflow/DataflowFloatVectorVolumeImpl.h"
 #include "Dataflow/DataflowBoolVolumeImpl.h"
 #include <Dataflow/DataflowBoolVolume.h>
+#include "Engine/VolumeTexture.h"
 
 namespace UE::DataflowVolume::Private
 {
@@ -174,6 +175,13 @@ namespace UE::DataflowVolume::Private
 
 	/* --------------------------------------------------------------------------------------------------------------- */
 
+	FBox FDataflowFloatVectorVolumeImpl::GetVolumeBoundingBox() const
+	{
+		return this->GetActiveVoxelBoundingBox();
+	}
+
+	/* --------------------------------------------------------------------------------------------------------------- */
+
 	TSharedPtr<FDataflowFloatVolumeImpl> FDataflowFloatVectorVolumeImpl::ComputeDivergence(const EDataflowVolumeAnalysisOutputName OutputName, const FString& CustomName, TSharedPtr<const FDataflowBoolVolumeImpl> InMaskVolumeImpl) const
 	{
 		using namespace openvdb;
@@ -347,7 +355,7 @@ namespace UE::DataflowVolume::Private
 
 		TSharedPtr<FDataflowFloatVectorVolumeImpl> FloatVectorVolumeImpl = MakeShared<FDataflowFloatVectorVolumeImpl>();
 
-		if (this->FloatVectorGrid->empty())
+		if (FloatVectorGrid->empty())
 		{
 			FloatVectorVolumeImpl->SetGrid(openvdb::VectorGrid::Ptr(), VectorGridType, GridName);
 			return FloatVectorVolumeImpl;
@@ -392,13 +400,57 @@ namespace UE::DataflowVolume::Private
 
 	/* --------------------------------------------------------------------------------------------------------------- */
 
+	bool FDataflowFloatVectorVolumeImpl::CreateVolumeTexture(UVolumeTexture* InVolumeTexture) const
+	{
+		if (InVolumeTexture == nullptr)
+		{
+			return false;
+		}
 
+		const EPixelFormat PixelFormat = EPixelFormat::PF_FloatRGB;
 
+		if (!FloatVectorGrid->empty())
+		{
+			openvdb::Coord ActiveVoxelDim = FloatVectorGrid->evalActiveVoxelDim();
+			openvdb::CoordBBox Bbox = FloatVectorGrid->evalActiveVoxelBoundingBox();
+			const int32 XMin = Bbox.min().x();
+			const int32 YMin = Bbox.min().y();
+			const int32 ZMin = Bbox.min().z();
 
+			openvdb::VectorGrid::Accessor Accessor = FloatVectorGrid->getAccessor();
 
+			auto QueryVoxel = [&Accessor, &XMin, &YMin, &ZMin](int32 PosX, int32 PosY, int32 PosZ, void* Value)
+			{
+				openvdb::Coord Coord(PosX + XMin, PosY + YMin, PosZ + ZMin);
+				FFloat16* const Voxel = static_cast<FFloat16*>(Value);
+
+				if (Accessor.isValueOn(Coord))
+				{
+					Voxel[0].Set(Accessor.getValue(Coord).x());
+					Voxel[1].Set(Accessor.getValue(Coord).y());
+					Voxel[2].Set(Accessor.getValue(Coord).z());
+					Voxel[3].Set(1.f);
+				}
+				else
+				{
+					Voxel[0].Set(0.f);
+					Voxel[1].Set(0.f);
+					Voxel[2].Set(0.f);
+					Voxel[3].Set(0.f);
+				}
+			};
+
+			const bool bSuccess = InVolumeTexture->UpdateSourceFromFunction(QueryVoxel, ActiveVoxelDim[0], ActiveVoxelDim[1], ActiveVoxelDim[2], TSF_RGBA16F);
+
+			if (!bSuccess)
+			{
+				return false;
+			}
 
-	
-
+			return true;
+		}
 
+		return false;
+	}
+}
 
-}
\ No newline at end of file
diff --git a/Engine/Plugins/Experimental/DataflowVolumetricPlugin/Source/DataflowVolumeCore/Private/Dataflow/DataflowFloatVectorVolumeImpl.h b/Engine/Plugins/Experimental/DataflowVolumetricPlugin/Source/DataflowVolumeCore/Private/Dataflow/DataflowFloatVectorVolumeImpl.h
--- a/Engine/Plugins/Experimental/DataflowVolumetricPlugin/Source/DataflowVolumeCore/Private/Dataflow/DataflowFloatVectorVolumeImpl.h
+++ b/Engine/Plugins/Experimental/DataflowVolumetricPlugin/Source/DataflowVolumeCore/Private/Dataflow/DataflowFloatVectorVolumeImpl.h
@@ -22,6 +22,11 @@ namespace UE::DataflowVolume::Private
 		openvdb::VectorGrid::Ptr FloatVectorGrid = nullptr;
 
 	public:
+		// Virtual functions
+		virtual void GetActiveVoxels(TArray<FBox>& OutActiveVoxels) const override;
+		virtual FBox GetVolumeBoundingBox() const override;
+		virtual bool CreateVolumeTexture(UVolumeTexture* InVolumeTexture) const override;
+
 		void SetGrid(const openvdb::VectorGrid::Ptr& InFloatVectorGrid, const openvdb::VecType InVecType, const FString& InGridName = "Gradient");
 
 		openvdb::VectorGrid::Ptr GetGrid() const;
@@ -34,8 +39,6 @@ namespace UE::DataflowVolume::Private
 
 		void VolumeSample(const TArray<FVector>& InPoints, TArray<FVector>& OutValues) const;
 
-		virtual void GetActiveVoxels(TArray<FBox>& OutActiveVoxels) const override;
-
 		TSharedPtr<FDataflowFloatVolumeImpl> ComputeDivergence(const EDataflowVolumeAnalysisOutputName OutputName, const FString& CustomName, TSharedPtr<const FDataflowBoolVolumeImpl> InMaskVolumeImpl) const;
 		TSharedPtr<FDataflowFloatVolumeImpl> ComputeMagnitude(const EDataflowVolumeAnalysisOutputName OutputName, const FString& CustomName, TSharedPtr<const FDataflowBoolVolumeImpl> InMaskVolumeImpl) const;
 		TSharedPtr<FDataflowFloatVectorVolumeImpl> ComputeCurl(const EDataflowVolumeAnalysisOutputName OutputName, const FString& CustomName, TSharedPtr<const FDataflowBoolVolumeImpl> InMaskVolumeImpl) const;
diff --git a/Engine/Plugins/Experimental/DataflowVolumetricPlugin/Source/DataflowVolumeCore/Private/Dataflow/DataflowFloatVolume.cpp b/Engine/Plugins/Experimental/DataflowVolumetricPlugin/Source/DataflowVolumeCore/Private/Dataflow/DataflowFloatVolume.cpp
--- a/Engine/Plugins/Experimental/DataflowVolumetricPlugin/Source/DataflowVolumeCore/Private/Dataflow/DataflowFloatVolume.cpp
+++ b/Engine/Plugins/Experimental/DataflowVolumetricPlugin/Source/DataflowVolumeCore/Private/Dataflow/DataflowFloatVolume.cpp
@@ -51,20 +51,6 @@ FVector FDataflowFloatVolume::GetVoxelSize() const
 
 /* --------------------------------------------------------------------------------------------------------------- */
 
-FBox FDataflowFloatVolume::GetVolumeBoundingBox() const
-{
-	using namespace UE::DataflowVolume::Private;
-
-	if (TSharedPtr<const FDataflowFloatVolumeImpl> FloatVolumeImpl = GetFloatVolumeImpl())
-	{
-		return FloatVolumeImpl->GetActiveVoxelBoundingBox();
-	}
-
-	return {};
-}
-
-/* --------------------------------------------------------------------------------------------------------------- */
-
 FDataflowFloatVolume FDataflowFloatVolume::CreateSphereSDF(
 	const float InVoxelSize, 
 	const float InRadius, 
@@ -433,3 +419,4 @@ FDataflowFloatVolume FDataflowFloatVolume::VolumeCombine(
 
 	return FDataflowFloatVolume();
 }
+
diff --git a/Engine/Plugins/Experimental/DataflowVolumetricPlugin/Source/DataflowVolumeCore/Private/Dataflow/DataflowFloatVolumeImpl.cpp b/Engine/Plugins/Experimental/DataflowVolumetricPlugin/Source/DataflowVolumeCore/Private/Dataflow/DataflowFloatVolumeImpl.cpp
--- a/Engine/Plugins/Experimental/DataflowVolumetricPlugin/Source/DataflowVolumeCore/Private/Dataflow/DataflowFloatVolumeImpl.cpp
+++ b/Engine/Plugins/Experimental/DataflowVolumetricPlugin/Source/DataflowVolumeCore/Private/Dataflow/DataflowFloatVolumeImpl.cpp
@@ -6,6 +6,7 @@
 #include "Dataflow/DataflowIntVolumeImpl.h"
 #include "Dataflow/DataflowBoolVolumeImpl.h"
 #include "Dataflow/DataflowVolumeUtils.h"
+#include "Engine/VolumeTexture.h"
 
 namespace UE::DataflowVolume::Private
 {
@@ -498,6 +499,13 @@ namespace UE::DataflowVolume::Private
 
 	/* --------------------------------------------------------------------------------------------------------------- */
 
+	FBox FDataflowFloatVolumeImpl::GetVolumeBoundingBox() const
+	{
+		return this->GetActiveVoxelBoundingBox();
+	}
+
+	/* --------------------------------------------------------------------------------------------------------------- */
+
 	typedef openvdb::FloatGrid GridT;
 	using TreeT = typename GridT::TreeType;
 	using BoolTreeT = typename TreeT::template ValueConverter<bool>::Type;
@@ -1188,9 +1196,51 @@ namespace UE::DataflowVolume::Private
 
 	/* --------------------------------------------------------------------------------------------------------------- */
 
+	bool FDataflowFloatVolumeImpl::CreateVolumeTexture(UVolumeTexture* InVolumeTexture) const
+	{
+		if (InVolumeTexture == nullptr)
+		{
+			return false;
+		}
+
+		const EPixelFormat PixelFormat = EPixelFormat::PF_R16F;
+
+		if (!FloatGrid->empty())
+		{
+			openvdb::Coord ActiveVoxelDim = FloatGrid->evalActiveVoxelDim();
+			openvdb::CoordBBox Bbox = FloatGrid->evalActiveVoxelBoundingBox();
+			const int32 XMin = Bbox.min().x();
+			const int32 YMin = Bbox.min().y();
+			const int32 ZMin = Bbox.min().z();
+
+			openvdb::FloatGrid::Accessor Accessor = FloatGrid->getAccessor();
 
+			auto QueryVoxel = [&Accessor, &XMin, &YMin, &ZMin](int32 PosX, int32 PosY, int32 PosZ, void* Value)
+			{
+				openvdb::Coord Coord(PosX + XMin, PosY + YMin, PosZ + ZMin);
+				FFloat16* const Voxel = static_cast<FFloat16*>(Value);
 
+				if (Accessor.isValueOn(Coord))
+				{
+					Voxel[0].Set(Accessor.getValue(Coord));
+				}
+				else
+				{
+					Voxel[0].Set(0.f);
+				}
+			};
 
+			const bool bSuccess = InVolumeTexture->UpdateSourceFromFunction(QueryVoxel, ActiveVoxelDim[0], ActiveVoxelDim[1], ActiveVoxelDim[2], TSF_R16F);
 
+			if (!bSuccess)
+			{
+				return false;
+			}
+
+			return true;
+		}
+
+		return false;
+	}
+}
 
-}
\ No newline at end of file
diff --git a/Engine/Plugins/Experimental/DataflowVolumetricPlugin/Source/DataflowVolumeCore/Private/Dataflow/DataflowFloatVolumeImpl.h b/Engine/Plugins/Experimental/DataflowVolumetricPlugin/Source/DataflowVolumeCore/Private/Dataflow/DataflowFloatVolumeImpl.h
--- a/Engine/Plugins/Experimental/DataflowVolumetricPlugin/Source/DataflowVolumeCore/Private/Dataflow/DataflowFloatVolumeImpl.h
+++ b/Engine/Plugins/Experimental/DataflowVolumetricPlugin/Source/DataflowVolumeCore/Private/Dataflow/DataflowFloatVolumeImpl.h
@@ -28,6 +28,12 @@ namespace UE::DataflowVolume::Private
 		openvdb::FloatGrid::Ptr FloatGrid = nullptr;
 
 	public:
+		// Virtual functions
+		virtual FString VolumeInfo() const override;
+		virtual void GetActiveVoxels(TArray<FBox>& OutActiveVoxels) const override;
+		virtual FBox GetVolumeBoundingBox() const override;
+		virtual bool CreateVolumeTexture(UVolumeTexture* InVolumeTexture) const override;
+
 		openvdb::FloatGrid::Ptr GetGrid() const;
 
 		void SetGrid(const openvdb::FloatGrid::Ptr& InFloatGrid, const openvdb::GridClass InGridClass, const FString& InGridName = TEXT("LevelSet"));
@@ -36,8 +42,6 @@ namespace UE::DataflowVolume::Private
 
 		FString GetGridClass() const;
 
-		virtual FString VolumeInfo() const override;
-
 		static TSharedPtr<FDataflowFloatVolumeImpl> CreateSphereSDF(
 			const float InVoxelSize,
 			const float InRadius,
@@ -98,8 +102,6 @@ namespace UE::DataflowVolume::Private
 
 		void ConvertVolumeToMeshDescription(const float InIsoValue, const float InAdaptivity, FMeshDescription& OutMeshDescription) const;
 
-		virtual void GetActiveVoxels(TArray<FBox>& OutActiveVoxels) const override;
-
 		void UniformVolumeScatter(
 			const int32 InMinNumberOfPoints,
 			const int32 InMaxNumberOfPoints,
@@ -141,4 +143,5 @@ namespace UE::DataflowVolume::Private
 			const bool InPruneDegenerateTiles,
 			const float InPruneTol) const;
 	};
-}
\ No newline at end of file
+}
+
diff --git a/Engine/Plugins/Experimental/DataflowVolumetricPlugin/Source/DataflowVolumeCore/Private/Dataflow/DataflowIntVolume.cpp b/Engine/Plugins/Experimental/DataflowVolumetricPlugin/Source/DataflowVolumeCore/Private/Dataflow/DataflowIntVolume.cpp
--- a/Engine/Plugins/Experimental/DataflowVolumetricPlugin/Source/DataflowVolumeCore/Private/Dataflow/DataflowIntVolume.cpp
+++ b/Engine/Plugins/Experimental/DataflowVolumetricPlugin/Source/DataflowVolumeCore/Private/Dataflow/DataflowIntVolume.cpp
@@ -16,27 +16,34 @@ FDataflowIntVolume::FDataflowIntVolume(TSharedPtr<UE::DataflowVolume::Private::F
 	static_assert(sizeof(FDataflowVolume) == sizeof(*this));
 }
 
-TSharedPtr<const UE::DataflowVolume::Private::FDataflowIntVolumeImpl> FDataflowIntVolume::GetIntVolume() const
+/* --------------------------------------------------------------------------------------------------------------- */
+
+TSharedPtr<const UE::DataflowVolume::Private::FDataflowIntVolumeImpl> FDataflowIntVolume::GetIntVolumeImpl() const
 {
 	return StaticCastSharedPtr<const UE::DataflowVolume::Private::FDataflowIntVolumeImpl>(Volume);
 }
 
-TSharedPtr<UE::DataflowVolume::Private::FDataflowIntVolumeImpl> FDataflowIntVolume::GetIntVolume()
+/* --------------------------------------------------------------------------------------------------------------- */
+
+TSharedPtr<UE::DataflowVolume::Private::FDataflowIntVolumeImpl> FDataflowIntVolume::GetIntVolumeImpl()
 {
 	return StaticCastSharedPtr<UE::DataflowVolume::Private::FDataflowIntVolumeImpl>(Volume);
 }
 
+/* --------------------------------------------------------------------------------------------------------------- */
 
 void FDataflowIntVolume::SetIntVolume(TSharedPtr<UE::DataflowVolume::Private::FDataflowIntVolumeImpl> InIntVolume)
 {
 	Volume = InIntVolume;
 }
 
+/* --------------------------------------------------------------------------------------------------------------- */
+
 FVector FDataflowIntVolume::GetVoxelSize() const
 {
 	using namespace UE::DataflowVolume::Private;
 
-	if (TSharedPtr<const FDataflowIntVolumeImpl> IntVolume = GetIntVolume())
+	if (TSharedPtr<const FDataflowIntVolumeImpl> IntVolume = GetIntVolumeImpl())
 	{
 		return IntVolume->GetVoxelSize();
 	}
@@ -44,17 +51,7 @@ FVector FDataflowIntVolume::GetVoxelSize() const
 	return {};
 }
 
-FBox FDataflowIntVolume::GetVolumeBoundingBox() const
-{
-	using namespace UE::DataflowVolume::Private;
-
-	if (TSharedPtr<const FDataflowIntVolumeImpl> IntVolume = GetIntVolume())
-	{
-		return IntVolume->GetActiveVoxelBoundingBox();
-	}
-
-	return {};
-}
+/* --------------------------------------------------------------------------------------------------------------- */
 
 void FDataflowIntVolume::VolumeSample(
 	const TArray<FVector>& InPoints,
@@ -62,10 +59,11 @@ void FDataflowIntVolume::VolumeSample(
 {
 	using namespace UE::DataflowVolume::Private;
 
-	if (TSharedPtr<const FDataflowIntVolumeImpl> IntVolume = GetIntVolume())
+	if (TSharedPtr<const FDataflowIntVolumeImpl> IntVolume = GetIntVolumeImpl())
 	{
 		IntVolume->VolumeSample(
 			InPoints,
 			OutValues);
 	}
 }
+
diff --git a/Engine/Plugins/Experimental/DataflowVolumetricPlugin/Source/DataflowVolumeCore/Private/Dataflow/DataflowIntVolumeImpl.cpp b/Engine/Plugins/Experimental/DataflowVolumetricPlugin/Source/DataflowVolumeCore/Private/Dataflow/DataflowIntVolumeImpl.cpp
--- a/Engine/Plugins/Experimental/DataflowVolumetricPlugin/Source/DataflowVolumeCore/Private/Dataflow/DataflowIntVolumeImpl.cpp
+++ b/Engine/Plugins/Experimental/DataflowVolumetricPlugin/Source/DataflowVolumeCore/Private/Dataflow/DataflowIntVolumeImpl.cpp
@@ -1,9 +1,8 @@
 // Copyright Epic Games, Inc. All Rights Reserved.
 
 #include "Dataflow/DataflowIntVolumeImpl.h"
-//#include "Dataflow/OpenVDB.h"
-//#include "Dataflow/DataflowVolumeUtils.h"
-//#include <limits>
+#include <Chaos/ArrayND.h>
+#include "Engine/VolumeTexture.h"
 
 namespace UE::DataflowVolume::Private
 {
@@ -18,11 +17,15 @@ namespace UE::DataflowVolume::Private
 		IntGrid = InIntGrid;
 	}
 
+	/* --------------------------------------------------------------------------------------------------------------- */
+
 	openvdb::Int32Grid::Ptr FDataflowIntVolumeImpl::GetGrid() const
 	{
 		return IntGrid;
 	}
 
+	/* --------------------------------------------------------------------------------------------------------------- */
+
 	FString FDataflowIntVolumeImpl::GetGridType() const
 	{
 		if (IntGrid)
@@ -34,6 +37,8 @@ namespace UE::DataflowVolume::Private
 		return {};
 	}
 
+	/* --------------------------------------------------------------------------------------------------------------- */
+
 	FString FDataflowIntVolumeImpl::VolumeInfo() const
 	{
 		FString BaseString = FDataflowVolumeImpl::VolumeInfo();
@@ -52,6 +57,8 @@ namespace UE::DataflowVolume::Private
 		return OutputStr;
 	}
 
+	/* --------------------------------------------------------------------------------------------------------------- */
+
 	void FDataflowIntVolumeImpl::VolumeSample(const TArray<FVector>& InPoints, TArray<int32>& OutValues) const
 	{
 		using namespace openvdb;
@@ -75,6 +82,8 @@ namespace UE::DataflowVolume::Private
 		return;
 	}
 
+	/* --------------------------------------------------------------------------------------------------------------- */
+
 	void FDataflowIntVolumeImpl::GetActiveVoxels(TArray<FBox>& OutActiveVoxels) const
 	{
 		const openvdb::Index64 ActiveVoxelCount = IntGrid->activeVoxelCount();
@@ -103,4 +112,61 @@ namespace UE::DataflowVolume::Private
 			}
 		}
 	}
+
+	/* --------------------------------------------------------------------------------------------------------------- */
+
+	FBox FDataflowIntVolumeImpl::GetVolumeBoundingBox() const
+	{
+		return this->GetActiveVoxelBoundingBox();
+	}
+
+	/* --------------------------------------------------------------------------------------------------------------- */
+
+	bool FDataflowIntVolumeImpl::CreateVolumeTexture(UVolumeTexture* InVolumeTexture) const
+	{
+		if (InVolumeTexture == nullptr)
+		{
+			return false;
+		}
+
+		const EPixelFormat PixelFormat = EPixelFormat::PF_R16F;
+
+		if (!IntGrid->empty())
+		{
+			openvdb::Coord ActiveVoxelDim = IntGrid->evalActiveVoxelDim();
+			openvdb::CoordBBox Bbox = IntGrid->evalActiveVoxelBoundingBox();
+			const int32 XMin = Bbox.min().x();
+			const int32 YMin = Bbox.min().y();
+			const int32 ZMin = Bbox.min().z();
+
+			openvdb::Int32Grid::Accessor Accessor = IntGrid->getAccessor();
+
+			auto QueryVoxel = [&Accessor, &XMin, &YMin, &ZMin](int32 PosX, int32 PosY, int32 PosZ, void* Value)
+			{
+				openvdb::Coord Coord(PosX + XMin, PosY + YMin, PosZ + ZMin);
+				FFloat16* const Voxel = static_cast<FFloat16*>(Value);
+
+				if (Accessor.isValueOn(Coord))
+				{
+					Voxel[0].Set(Accessor.getValue(Coord));
+				}
+				else
+				{
+					Voxel[0].Set(0.f);
+				}
+			};
+
+			const bool bSuccess = InVolumeTexture->UpdateSourceFromFunction(QueryVoxel, ActiveVoxelDim[0], ActiveVoxelDim[1], ActiveVoxelDim[2], TSF_R16F);
+
+			if (!bSuccess)
+			{
+				return false;
+			}
+
+			return true;
+		}
+
+		return false;
+	}
 }
+
diff --git a/Engine/Plugins/Experimental/DataflowVolumetricPlugin/Source/DataflowVolumeCore/Private/Dataflow/DataflowIntVolumeImpl.h b/Engine/Plugins/Experimental/DataflowVolumetricPlugin/Source/DataflowVolumeCore/Private/Dataflow/DataflowIntVolumeImpl.h
--- a/Engine/Plugins/Experimental/DataflowVolumetricPlugin/Source/DataflowVolumeCore/Private/Dataflow/DataflowIntVolumeImpl.h
+++ b/Engine/Plugins/Experimental/DataflowVolumetricPlugin/Source/DataflowVolumeCore/Private/Dataflow/DataflowIntVolumeImpl.h
@@ -17,16 +17,19 @@ namespace UE::DataflowVolume::Private
 		openvdb::Int32Grid::Ptr IntGrid = nullptr;
 
 	public:
-		void SetGrid(const openvdb::Int32Grid::Ptr& InIntGrid, const FString& InGridName = "");
+		// Virtual functions
+		virtual FString VolumeInfo() const override;
+		virtual void GetActiveVoxels(TArray<FBox>& OutActiveVoxels) const override;
+		virtual FBox GetVolumeBoundingBox() const override;
+		virtual bool CreateVolumeTexture(UVolumeTexture* InVolumeTexture) const override;
 
 		openvdb::Int32Grid::Ptr GetGrid() const;
 
-		FString GetGridType() const;
+		void SetGrid(const openvdb::Int32Grid::Ptr& InIntGrid, const FString& InGridName = "");
 
-		virtual FString VolumeInfo() const override;
+		FString GetGridType() const;
 
 		void VolumeSample(const TArray<FVector>& InPoints, TArray<int32>& OutValues) const;
 
-		virtual void GetActiveVoxels(TArray<FBox>& OutActiveVoxels) const override;
 	};
 }
diff --git a/Engine/Plugins/Experimental/DataflowVolumetricPlugin/Source/DataflowVolumeCore/Private/Dataflow/DataflowVolume.cpp b/Engine/Plugins/Experimental/DataflowVolumetricPlugin/Source/DataflowVolumeCore/Private/Dataflow/DataflowVolume.cpp
--- a/Engine/Plugins/Experimental/DataflowVolumetricPlugin/Source/DataflowVolumeCore/Private/Dataflow/DataflowVolume.cpp
+++ b/Engine/Plugins/Experimental/DataflowVolumetricPlugin/Source/DataflowVolumeCore/Private/Dataflow/DataflowVolume.cpp
@@ -12,6 +12,8 @@ FDataflowVolume::FDataflowVolume(TSharedPtr<UE::DataflowVolume::Private::FDatafl
 	, Type(InType)
 {}
 
+/* --------------------------------------------------------------------------------------------------------------- */
+
 FString FDataflowVolume::VolumeInfo() const
 {
 	if (Volume)
@@ -22,6 +24,23 @@ FString FDataflowVolume::VolumeInfo() const
 	return {};
 }
 
+/* --------------------------------------------------------------------------------------------------------------- */
+
+int32 FDataflowVolume::GetNumActiveVoxels() const
+{
+	if (Volume)
+	{
+		TArray<FBox> OutActiveVoxels;
+		Volume->GetActiveVoxels(OutActiveVoxels);
+
+		return OutActiveVoxels.Num();
+	}
+
+	return 0;
+}
+
+/* --------------------------------------------------------------------------------------------------------------- */
+
 void FDataflowVolume::GetActiveVoxels(TArray<FBox>& OutActiveVoxels) const
 {
 	if (Volume)
@@ -30,5 +49,27 @@ void FDataflowVolume::GetActiveVoxels(TArray<FBox>& OutActiveVoxels) const
 	}
 }
 
+/* --------------------------------------------------------------------------------------------------------------- */
+
+FBox FDataflowVolume::GetVolumeBoundingBox() const
+{
+	if (Volume)
+	{
+		return Volume->GetVolumeBoundingBox();
+	}
+
+	return {};
+}
+
+/* --------------------------------------------------------------------------------------------------------------- */
+
+bool FDataflowVolume::CreateVolumeTexture(UVolumeTexture* InVolumeTexture) const
+{
+	if (Volume)
+	{
+		return Volume->CreateVolumeTexture(InVolumeTexture);
+	}
 
+	return false;
+}
 
diff --git a/Engine/Plugins/Experimental/DataflowVolumetricPlugin/Source/DataflowVolumeCore/Private/Dataflow/DataflowVolumeImpl.h b/Engine/Plugins/Experimental/DataflowVolumetricPlugin/Source/DataflowVolumeCore/Private/Dataflow/DataflowVolumeImpl.h
--- a/Engine/Plugins/Experimental/DataflowVolumetricPlugin/Source/DataflowVolumeCore/Private/Dataflow/DataflowVolumeImpl.h
+++ b/Engine/Plugins/Experimental/DataflowVolumetricPlugin/Source/DataflowVolumeCore/Private/Dataflow/DataflowVolumeImpl.h
@@ -6,12 +6,13 @@
 #include "Math/UnrealMath.h"
 #include "Containers/UnrealString.h"
 
+class UVolumeTexture;
+
 namespace UE::DataflowVolume::Private
 {
 	/* --------------------------------------------------------------------------------------------------------------- */
 	/* FDataflowVolumeImpl */
 	/* --------------------------------------------------------------------------------------------------------------- */
-	
 	class FDataflowVolumeImpl
 	{
 	private:
@@ -28,6 +29,8 @@ namespace UE::DataflowVolume::Private
 
 		virtual void GetActiveVoxels(TArray<FBox>& OutActiveVoxels) const = 0;
 
+		virtual FBox GetVolumeBoundingBox() const = 0;
+
 		FIntVector3 GetActiveVoxelDim() const;
 
 		FBox GetActiveVoxelBoundingBox() const;
@@ -38,6 +41,8 @@ namespace UE::DataflowVolume::Private
 
 		virtual ~FDataflowVolumeImpl() {}
 
+		virtual bool CreateVolumeTexture(UVolumeTexture* InVolumeTexture) const = 0;
+
 	protected:
 		void SetGrid(const openvdb::GridBase::Ptr& InGrid, const FString& InGridName);
 	};
diff --git a/Engine/Plugins/Experimental/DataflowVolumetricPlugin/Source/DataflowVolumeNodes/Private/Dataflow/DataflowVolumeNodes.cpp b/Engine/Plugins/Experimental/DataflowVolumetricPlugin/Source/DataflowVolumeNodes/Private/Dataflow/DataflowVolumeNodes.cpp
--- a/Engine/Plugins/Experimental/DataflowVolumetricPlugin/Source/DataflowVolumeNodes/Private/Dataflow/DataflowVolumeNodes.cpp
+++ b/Engine/Plugins/Experimental/DataflowVolumetricPlugin/Source/DataflowVolumeNodes/Private/Dataflow/DataflowVolumeNodes.cpp
@@ -13,6 +13,7 @@
 #include "Dataflow/DataflowRenderingViewMode.h"
 #endif
 #include "GeometryCollection/Facades/PointsCollectionFacade.h"
+#include "Engine/VolumeTexture.h"
 
 #include UE_INLINE_GENERATED_CPP_BY_NAME(DataflowVolumeNodes)
 
@@ -43,6 +44,9 @@ namespace UE::Dataflow
 
 		// Visualize
 		DATAFLOW_NODE_REGISTER_CREATION_FACTORY(FVolumeSliceDataflowNode);
+
+		// Terminal
+		DATAFLOW_NODE_REGISTER_CREATION_FACTORY(FVolumeTextureTerminalNode);
 	}
 }
 
@@ -1971,16 +1975,72 @@ void FVolumeCombineDataflowNode::DebugDraw(UE::Dataflow::FContext& Context, IDat
 
 /* --------------------------------------------------------------------------------------------- */
 
+namespace DataflowVolumeNodes::Private
+{
+	void UpdateVolumeTextureFromVolume(UVolumeTexture* InVolumeTexture, const FDataflowVolume& InVolume)
+	{
+#if WITH_EDITOR
+		InVolumeTexture->PreEditChange(nullptr);
+#endif
+
+#if WITH_EDITORONLY_DATA
+		InVolume.CreateVolumeTexture(InVolumeTexture);
 
+		if (!InVolumeTexture)
+		{
+			return;
+		}
+#endif
+		InVolumeTexture->UpdateResource();
 
+#if WITH_EDITOR
+		InVolumeTexture->PostEditChange();
+#endif
+	}
+}
 
+FVolumeTextureTerminalNode::FVolumeTextureTerminalNode(const UE::Dataflow::FNodeParameters& InParam, FGuid InGuid)
+	: FDataflowTerminalNode(InParam, InGuid)
+{
+	static const FName MainTypeGroup("Main");
 
+	RegisterInputConnection(&Volume)
+		.SetTypeDependencyGroup(MainTypeGroup);
 
+	RegisterInputConnection(&VolumeTextureAsset);
 
+	RegisterOutputConnection(&Volume, &Volume)
+		.SetTypeDependencyGroup(MainTypeGroup);
+}
 
+void FVolumeTextureTerminalNode::Evaluate(UE::Dataflow::FContext& Context) const
+{
+	ForwardInput(Context, &Volume, &Volume);
+}
 
 
+void FVolumeTextureTerminalNode::SetAssetValue(TObjectPtr<UObject> Asset, UE::Dataflow::FContext& Context) const
+{
+	UVolumeTexture* AssetToSet = Cast<UVolumeTexture>(Asset.Get());
+	if (!AssetToSet)
+	{
+		AssetToSet = GetValue(Context, &VolumeTextureAsset);
+	}
 
+	if (AssetToSet)
+	{
+		const FDataflowVolume& InVolume = GetValue(Context, &Volume);
 
+		if (InVolume.GetNumActiveVoxels() > 0)
+		{
+			DataflowVolumeNodes::Private::UpdateVolumeTextureFromVolume(AssetToSet, InVolume);
+		}
+		else
+		{
+			Context.Warning(TEXT("Volume doesn't contain data, VolumeTexture will be empty."));
+		}
+	}
+}
 
+/* --------------------------------------------------------------------------------------------- */
 
diff --git a/Engine/Plugins/Experimental/DataflowVolumetricPlugin/Source/DataflowVolumeNodes/Private/Dataflow/DataflowVolumeNodes.h b/Engine/Plugins/Experimental/DataflowVolumetricPlugin/Source/DataflowVolumeNodes/Private/Dataflow/DataflowVolumeNodes.h
--- a/Engine/Plugins/Experimental/DataflowVolumetricPlugin/Source/DataflowVolumeNodes/Private/Dataflow/DataflowVolumeNodes.h
+++ b/Engine/Plugins/Experimental/DataflowVolumetricPlugin/Source/DataflowVolumeNodes/Private/Dataflow/DataflowVolumeNodes.h
@@ -856,6 +856,29 @@ public:
 	FVolumeCombineDataflowNode(const UE::Dataflow::FNodeParameters& InParam, FGuid InGuid = FGuid::NewGuid());
 };
 
+/*
+* Terminal node to a save volume texture
+*/
+USTRUCT(meta = (DataflowTerminal))
+struct FVolumeTextureTerminalNode : public FDataflowTerminalNode
+{
+	GENERATED_USTRUCT_BODY()
+	DATAFLOW_NODE_DEFINE_INTERNAL(FVolumeTextureTerminalNode, "VolumeTextureTerminal", "Terminal", "")
+
+private:
+	UPROPERTY(meta = (DataflowInput, DataflowOutput, DataflowIntrinsic, DataflowPassthrough = "Volume"))
+	FDataflowVolumeTypes Volume;
+
+	UPROPERTY(EditAnywhere, Category = Asset, meta = (DataflowInput))
+	TObjectPtr<UVolumeTexture> VolumeTextureAsset = nullptr;
+
+	virtual void Evaluate(UE::Dataflow::FContext& Context) const override;
+	virtual void SetAssetValue(TObjectPtr<UObject> Asset, UE::Dataflow::FContext& Context) const override;
+
+public:
+	FVolumeTextureTerminalNode(const UE::Dataflow::FNodeParameters& InParam, FGuid InGuid = FGuid::NewGuid());
+};
+
 namespace UE::Dataflow
 {
 	void RegisterVolumeNodes();
