Dataflow:
- fixed minor issues pointed out by static analysis

diff --git a/Engine/Plugins/Experimental/DataflowVolumetricPlugin/Source/DataflowVolumeNodes/Private/Dataflow/DataflowVolumeNodes.cpp b/Engine/Plugins/Experimental/DataflowVolumetricPlugin/Source/DataflowVolumeNodes/Private/Dataflow/DataflowVolumeNodes.cpp
--- a/Engine/Plugins/Experimental/DataflowVolumetricPlugin/Source/DataflowVolumeNodes/Private/Dataflow/DataflowVolumeNodes.cpp
+++ b/Engine/Plugins/Experimental/DataflowVolumetricPlugin/Source/DataflowVolumeNodes/Private/Dataflow/DataflowVolumeNodes.cpp
@@ -1,11 +1,10 @@
 // Copyright Epic Games, Inc. All Rights Reserved.
 
 #include "Dataflow/DataflowVolumeNodes.h"
+#include "Dataflow/DataflowNodeFactory.h"
 #include "Dataflow/DataflowCore.h"
-//#include "UDynamicMesh.h"
 #include "StaticMeshAttributes.h"
 #include "PlanarCut.h"
-//#include "TransformTypes.h"
 #include "GeometryCollection/GeometryCollectionEngineConversion.h"
 #include "Dataflow/DataflowOverlay.h"
 #include "Dataflow/DataflowDebugDrawInterface.h"
@@ -1979,21 +1978,20 @@ namespace DataflowVolumeNodes::Private
 {
 	void UpdateVolumeTextureFromVolume(UVolumeTexture* InVolumeTexture, const FDataflowVolume& InVolume)
 	{
+		if (!InVolumeTexture)
+		{
+			return;
+		}
+
 #if WITH_EDITOR
 		InVolumeTexture->PreEditChange(nullptr);
-#endif
-
-#if WITH_EDITORONLY_DATA
-		InVolume.CreateVolumeTexture(InVolumeTexture);
 
-		if (!InVolumeTexture)
+		if (!InVolume.CreateVolumeTexture(InVolumeTexture))
 		{
 			return;
 		}
-#endif
 		InVolumeTexture->UpdateResource();
 
-#if WITH_EDITOR
 		InVolumeTexture->PostEditChange();
 #endif
 	}
