Dataflow:
- Added VolumeToSpheres_v2 node with an improved algorithm and exposed settings

diff --git a/Engine/Plugins/Experimental/DataflowVolumetricPlugin/Source/DataflowVolumeCore/Internal/Dataflow/DataflowFloatVolume.h b/Engine/Plugins/Experimental/DataflowVolumetricPlugin/Source/DataflowVolumeCore/Internal/Dataflow/DataflowFloatVolume.h
--- a/Engine/Plugins/Experimental/DataflowVolumetricPlugin/Source/DataflowVolumeCore/Internal/Dataflow/DataflowFloatVolume.h
+++ b/Engine/Plugins/Experimental/DataflowVolumetricPlugin/Source/DataflowVolumeCore/Internal/Dataflow/DataflowFloatVolume.h
@@ -63,7 +63,7 @@ public:
 		const float InIsoValue,
 		FDataflowIntVolume& OutFaceIndexVolume);
 
-	/** Compute sphere packing in an SDF volume */
+	/** Compute sphere packing in an SDF volume using original algorithm */
 	DATAFLOWVOLUMECORE_API void VolumeToSpheres(
 		const int32 InMinSphereCount,
 		const int32 InMaxSphereCount,
@@ -72,8 +72,25 @@ public:
 		const float InMaxRadius,
 		const float InIsoValue,
 		const int32 InInstanceCount,
-		TArray<FVector>& SphereCenters,
-		TArray<float>& SphereRadii) const;
+		TArray<FVector>& OutSphereCenters,
+		TArray<float>& OutSphereRadii) const;
+
+	/** Compute sphere packing in an SDF volume using improved algorithm */
+	DATAFLOWVOLUMECORE_API void VolumeToSpheresImproved(
+		const int32 InMinSphereCount,
+		const int32 InMaxSphereCount,
+		const bool InOverlapping,
+		const float InMinRadius,
+		const float InMaxRadius,
+		const float InIsoValue,
+		const int32 InInstanceCount,
+		const int32 InRandomSeed,
+		const EDataflowVolumeScatterType InScatterType,
+		const float InSpread,
+		const float InMinNumberOfPointsPerVoxel,
+		const float InMaxNumberOfPointsPerVoxel,
+		TArray<FVector>& OutSphereCenters,
+		TArray<float>& OutSphereRadii) const;
 
 	/** Convert SDF to Fog Volume */
 	DATAFLOWVOLUMECORE_API FDataflowFloatVolume CovertSDFToFogVolume(const bool InPruneTolerance, const float InTolerance, const bool InFloodFillOutput, const bool InActivateInteriror) const;
diff --git a/Engine/Plugins/Experimental/DataflowVolumetricPlugin/Source/DataflowVolumeCore/Internal/Dataflow/DataflowVolumeNodeEnums.h b/Engine/Plugins/Experimental/DataflowVolumetricPlugin/Source/DataflowVolumeCore/Internal/Dataflow/DataflowVolumeNodeEnums.h
--- a/Engine/Plugins/Experimental/DataflowVolumetricPlugin/Source/DataflowVolumeCore/Internal/Dataflow/DataflowVolumeNodeEnums.h
+++ b/Engine/Plugins/Experimental/DataflowVolumetricPlugin/Source/DataflowVolumeCore/Internal/Dataflow/DataflowVolumeNodeEnums.h
@@ -167,4 +167,14 @@ enum class EDataflowVolumeAnalysisOutputName : uint8
 {
 	AppendOperationName UMETA(DisplayName = "Append Operation Name"),
 	CustomName UMETA(DisplayName = "Custom Name"),
-};
\ No newline at end of file
+};
+
+UENUM(BlueprintType)
+enum class EDataflowVolumeScatterType : uint8
+{
+	Uniform UMETA(DisplayName = "Uniform"),
+	DenseUniform UMETA(DisplayName = "Dense Uniform"),
+	NonUniform UMETA(DisplayName = "Non Uniform")
+};
+
+
diff --git a/Engine/Plugins/Experimental/DataflowVolumetricPlugin/Source/DataflowVolumeCore/Private/Dataflow/DataflowFloatVolume.cpp b/Engine/Plugins/Experimental/DataflowVolumetricPlugin/Source/DataflowVolumeCore/Private/Dataflow/DataflowFloatVolume.cpp
--- a/Engine/Plugins/Experimental/DataflowVolumetricPlugin/Source/DataflowVolumeCore/Private/Dataflow/DataflowFloatVolume.cpp
+++ b/Engine/Plugins/Experimental/DataflowVolumetricPlugin/Source/DataflowVolumeCore/Private/Dataflow/DataflowFloatVolume.cpp
@@ -143,8 +143,8 @@ void FDataflowFloatVolume::VolumeToSpheres(
 	const float InMaxRadius,
 	const float InIsoValue,
 	const int32 InInstanceCount,
-	TArray<FVector>& SphereCenters,
-	TArray<float>& SphereRadii) const
+	TArray<FVector>& OutSphereCenters,
+	TArray<float>& OutSphereRadii) const
 {
 	using namespace UE::DataflowVolume::Private;
 
@@ -158,8 +158,48 @@ void FDataflowFloatVolume::VolumeToSpheres(
 			InMaxRadius,
 			InIsoValue,
 			InInstanceCount,
-			SphereCenters,
-			SphereRadii);
+			OutSphereCenters,
+			OutSphereRadii);
+	}
+}
+
+/* --------------------------------------------------------------------------------------------------------------- */
+
+void FDataflowFloatVolume::VolumeToSpheresImproved(
+	const int32 InMinSphereCount,
+	const int32 InMaxSphereCount,
+	const bool InOverlapping,
+	const float InMinRadius,
+	const float InMaxRadius,
+	const float InIsoValue,
+	const int32 InInstanceCount,
+	const int32 InRandomSeed,
+	const EDataflowVolumeScatterType InScatterType,
+	const float InSpread,
+	const float InMinNumberOfPointsPerVoxel,
+	const float InMaxNumberOfPointsPerVoxel,
+	TArray<FVector>& OutSphereCenters,
+	TArray<float>& OutSphereRadii) const
+{
+	using namespace UE::DataflowVolume::Private;
+
+	if (TSharedPtr<const FDataflowFloatVolumeImpl> FloatVolumeImpl = GetFloatVolumeImpl())
+	{
+		FloatVolumeImpl->VolumeToSpheresImproved(
+			InMinSphereCount,
+			InMaxSphereCount,
+			InOverlapping,
+			InMinRadius,
+			InMaxRadius,
+			InIsoValue,
+			InInstanceCount,
+			InRandomSeed,
+			InScatterType,
+			InSpread,
+			InMinNumberOfPointsPerVoxel,
+			InMaxNumberOfPointsPerVoxel,
+			OutSphereCenters,
+			OutSphereRadii);
 	}
 }
 
diff --git a/Engine/Plugins/Experimental/DataflowVolumetricPlugin/Source/DataflowVolumeCore/Private/Dataflow/DataflowFloatVolumeImpl.cpp b/Engine/Plugins/Experimental/DataflowVolumetricPlugin/Source/DataflowVolumeCore/Private/Dataflow/DataflowFloatVolumeImpl.cpp
--- a/Engine/Plugins/Experimental/DataflowVolumetricPlugin/Source/DataflowVolumeCore/Private/Dataflow/DataflowFloatVolumeImpl.cpp
+++ b/Engine/Plugins/Experimental/DataflowVolumetricPlugin/Source/DataflowVolumeCore/Private/Dataflow/DataflowFloatVolumeImpl.cpp
@@ -6,6 +6,7 @@
 #include "Dataflow/DataflowIntVolumeImpl.h"
 #include "Dataflow/DataflowBoolVolumeImpl.h"
 #include "Dataflow/DataflowVolumeUtils.h"
+#include "Dataflow/DataflowVolumeAlgo.h"
 #include "Engine/VolumeTexture.h"
 
 namespace UE::DataflowVolume::Private
@@ -308,8 +309,8 @@ namespace UE::DataflowVolume::Private
 		float InMaxRadius,
 		const float InIsoValue,
 		const int32 InInstanceCount,
-		TArray<FVector>& SphereCenters,
-		TArray<float>& SphereRadii) const
+		TArray<FVector>& OutSphereCenters,
+		TArray<float>& OutSphereRadii) const
 	{
 		if (FloatGrid)
 		{
@@ -326,13 +327,57 @@ namespace UE::DataflowVolume::Private
 				Spheres, SphereCount, InOverlapping, InMinRadius, InMaxRadius, InIsoValue, InInstanceCount);
 
 			const int32 NumSpheres = Spheres.size();
-			SphereCenters.AddUninitialized(NumSpheres);
-			SphereRadii.AddUninitialized(NumSpheres);
+			OutSphereCenters.AddUninitialized(NumSpheres);
+			OutSphereRadii.AddUninitialized(NumSpheres);
 
 			for (int32 Idx = 0; Idx < NumSpheres; ++Idx)
 			{
-				SphereCenters[Idx] = FVector(Spheres[Idx].x(), Spheres[Idx].y(), Spheres[Idx].z());
-				SphereRadii[Idx] = Spheres[Idx].w();
+				OutSphereCenters[Idx] = FVector(Spheres[Idx].x(), Spheres[Idx].y(), Spheres[Idx].z());
+				OutSphereRadii[Idx] = Spheres[Idx].w();
+			}
+		}
+	}
+
+	/* --------------------------------------------------------------------------------------------------------------- */
+
+	void FDataflowFloatVolumeImpl::VolumeToSpheresImproved(const int32 InMinSphereCount,
+		const int32 InMaxSphereCount,
+		const bool InOverlapping,
+		const float InMinRadius,
+		float InMaxRadius,
+		const float InIsoValue,
+		const int32 InInstanceCount,
+		const int32 InRandomSeed,
+		const EDataflowVolumeScatterType InScatterType,
+		const float InSpread,
+		const float InMinNumberOfPointsPerVoxel,
+		const float InMaxNumberOfPointsPerVoxel,
+		TArray<FVector>& OutSphereCenters,
+		TArray<float>& OutSphereRadii) const
+	{
+		if (FloatGrid)
+		{
+			openvdb::FloatGrid::ConstPtr GridPtr = openvdb::gridConstPtrCast<openvdb::FloatGrid>(FloatGrid);
+
+			std::vector<openvdb::Vec4s> Spheres;
+			const openvdb::Vec2i SphereCount(InMinSphereCount, InMaxSphereCount);
+			if (InMaxRadius <= InMinRadius)
+			{
+				InMaxRadius = InMinRadius + (float)1e-5;
+			}
+
+			FillWithSpheres(*GridPtr,
+				Spheres, SphereCount, InOverlapping, InMinRadius, InMaxRadius, 
+				InIsoValue, InInstanceCount, InRandomSeed, InScatterType, InSpread, InMinNumberOfPointsPerVoxel, InMaxNumberOfPointsPerVoxel);
+
+			const int32 NumSpheres = Spheres.size();
+			OutSphereCenters.AddUninitialized(NumSpheres);
+			OutSphereRadii.AddUninitialized(NumSpheres);
+
+			for (int32 Idx = 0; Idx < NumSpheres; ++Idx)
+			{
+				OutSphereCenters[Idx] = FVector(Spheres[Idx].x(), Spheres[Idx].y(), Spheres[Idx].z());
+				OutSphereRadii[Idx] = Spheres[Idx].w();
 			}
 		}
 	}
diff --git a/Engine/Plugins/Experimental/DataflowVolumetricPlugin/Source/DataflowVolumeCore/Private/Dataflow/DataflowFloatVolumeImpl.h b/Engine/Plugins/Experimental/DataflowVolumetricPlugin/Source/DataflowVolumeCore/Private/Dataflow/DataflowFloatVolumeImpl.h
--- a/Engine/Plugins/Experimental/DataflowVolumetricPlugin/Source/DataflowVolumeCore/Private/Dataflow/DataflowFloatVolumeImpl.h
+++ b/Engine/Plugins/Experimental/DataflowVolumetricPlugin/Source/DataflowVolumeCore/Private/Dataflow/DataflowFloatVolumeImpl.h
@@ -82,8 +82,23 @@ namespace UE::DataflowVolume::Private
 			float InMaxRadius,
 			const float InIsoValue,
 			const int32 InInstanceCount,
-			TArray<FVector>& SphereCenters,
-			TArray<float>& SphereRadii) const;
+			TArray<FVector>& OutSphereCenters,
+			TArray<float>& OutSphereRadii) const;
+
+		void VolumeToSpheresImproved(const int32 InMinSphereCount,
+			const int32 InMaxSphereCount,
+			const bool InOverlapping,
+			const float InMinRadius,
+			float InMaxRadius,
+			const float InIsoValue,
+			const int32 InInstanceCount,
+			const int32 InRandomSeed,
+			const EDataflowVolumeScatterType InScatterType,
+			const float InSpread,
+			const float InMinNumberOfPointsPerVoxel,
+			const float InMaxNumberOfPointsPerVoxel,
+			TArray<FVector>& OutSphereCenters,
+			TArray<float>& OutSphereRadii) const;
 
 		TSharedPtr<FDataflowFloatVolumeImpl> CovertSDFToFogVolume(
 			const bool InPruneTolerance,
diff --git a/Engine/Plugins/Experimental/DataflowVolumetricPlugin/Source/DataflowVolumeCore/Private/Dataflow/DataflowVolumeAlgo.cpp b/Engine/Plugins/Experimental/DataflowVolumetricPlugin/Source/DataflowVolumeCore/Private/Dataflow/DataflowVolumeAlgo.cpp
new file mode 100644
--- /dev/null
+++ b/Engine/Plugins/Experimental/DataflowVolumetricPlugin/Source/DataflowVolumeCore/Private/Dataflow/DataflowVolumeAlgo.cpp
@@ -0,0 +1,4 @@
+// Copyright Epic Games, Inc. All Rights Reserved.
+
+#include "Dataflow/DataflowVolumeAlgo.h"
+
diff --git a/Engine/Plugins/Experimental/DataflowVolumetricPlugin/Source/DataflowVolumeCore/Private/Dataflow/DataflowVolumeAlgo.h b/Engine/Plugins/Experimental/DataflowVolumetricPlugin/Source/DataflowVolumeCore/Private/Dataflow/DataflowVolumeAlgo.h
new file mode 100644
--- /dev/null
+++ b/Engine/Plugins/Experimental/DataflowVolumetricPlugin/Source/DataflowVolumeCore/Private/Dataflow/DataflowVolumeAlgo.h
@@ -0,0 +1,261 @@
+// Copyright Epic Games, Inc. All Rights Reserved.
+
+#pragma once
+
+#include "Dataflow/OpenVDB.h"
+#include "Dataflow/DataflowVolumeNodeEnums.h"
+
+#include "Math/RandomStream.h"
+
+#include <tbb/blocked_range.h>
+#include <tbb/parallel_for.h>
+#include <tbb/parallel_reduce.h>
+
+#include <algorithm> // for std::min(), std::max()
+#include <cmath> // for std::sqrt()
+#include <limits> // for std::numeric_limits
+#include <memory>
+#include <random>
+#include <utility> // for std::pair
+#include <vector>
+
+#include <openvdb/tools/VolumeToSpheres.h>
+#include <openvdb/tools/LevelSetRebuild.h>
+#include <openvdb/tools/LevelSetUtil.h>
+#include <openvdb/util/NullInterrupter.h>
+
+// Ideas:
+// 
+// - expose random seed as input
+// - expose scatter type as input
+// - make sure that the computed radius for the scattered points never gets bigger than maxRadius
+// - move scattered points with phi < phi_threshold in the inverse gradient direction by a random amount
+// - option to don't clamp radius by closestpointonsurface
+//
+
+template<typename GridT, typename InterrupterT = openvdb::util::NullInterrupter>
+void FillWithSpheres(
+	const GridT& grid,
+	std::vector<openvdb::Vec4s>& spheres,
+	const openvdb::Vec2i& sphereCount,
+	bool overlapping,
+	float minRadius,
+	float maxRadius,
+	float isovalue,
+	int instanceCount,
+	int32 InRandomSeed,
+	EDataflowVolumeScatterType InScatterType,
+	float InSpread,
+	float InMinNumberOfPointsPerVoxel,
+	float InMaxNumberOfPointsPerVoxel,
+	InterrupterT* interrupter = nullptr)
+{
+	spheres.clear();
+
+	if (grid.empty()) return;
+
+	const int
+		minSphereCount = sphereCount[0],
+		maxSphereCount = sphereCount[1];
+	if ((minSphereCount > maxSphereCount) || (maxSphereCount < 1)) {
+		OPENVDB_LOG_WARN("fillWithSpheres: minimum sphere count ("
+			<< minSphereCount << ") exceeds maximum count (" << maxSphereCount << ")");
+		return;
+	}
+	spheres.reserve(maxSphereCount);
+
+	auto gridPtr = grid.copy(); // shallow copy
+
+	if (gridPtr->getGridClass() == openvdb::GRID_LEVEL_SET) {
+		// Clamp the isovalue to the level set's background value minus epsilon.
+		// (In a valid narrow-band level set, all voxels, including background voxels,
+		// have values less than or equal to the background value, so an isovalue
+		// greater than or equal to the background value would produce a mask with
+		// effectively infinite extent.)
+		isovalue = std::min(isovalue,
+			static_cast<float>(gridPtr->background() - openvdb::math::Tolerance<float>::value()));
+	}
+	else if (gridPtr->getGridClass() == openvdb::GRID_FOG_VOLUME) {
+		// Clamp the isovalue of a fog volume between epsilon and one,
+		// again to avoid a mask with infinite extent.  (Recall that
+		// fog volume voxel values vary from zero outside to one inside.)
+		isovalue = openvdb::math::Clamp(isovalue, openvdb::math::Tolerance<float>::value(), 1.f);
+	}
+
+	// ClosestSurfacePoint is inaccurate for small grids.
+	// Resample the input grid if it is too small.
+	auto numVoxels = gridPtr->activeVoxelCount();
+	if (numVoxels < 10000) {
+		const auto scale = 1.0 / openvdb::math::Cbrt(2.0 * 10000.0 / double(numVoxels));
+		auto scaledXform = gridPtr->transform().copy();
+		scaledXform->preScale(scale);
+
+		auto newGridPtr = openvdb::tools::levelSetRebuild(*gridPtr, isovalue,
+			openvdb::LEVEL_SET_HALF_WIDTH, openvdb::LEVEL_SET_HALF_WIDTH, scaledXform.get(), interrupter);
+
+		const auto newNumVoxels = newGridPtr->activeVoxelCount();
+		if (newNumVoxels > numVoxels) {
+			OPENVDB_LOG_DEBUG_RUNTIME("fillWithSpheres: resampled input grid from "
+				<< numVoxels << " voxel" << (numVoxels == 1 ? "" : "s")
+				<< " to " << newNumVoxels << " voxel" << (newNumVoxels == 1 ? "" : "s"));
+			gridPtr = newGridPtr;
+			numVoxels = newNumVoxels;
+		}
+	}
+
+	const bool addNarrowBandPoints = (numVoxels < 10000);
+	int instances = std::max(instanceCount, maxSphereCount);
+
+	using TreeT = typename GridT::TreeType;
+	using BoolTreeT = typename TreeT::template ValueConverter<bool>::Type;
+	using Int16TreeT = typename TreeT::template ValueConverter<openvdb::Int16>::Type;
+
+	using RandGen = std::mersenne_twister_engine<uint32_t, 32, 351, 175, 19,
+		0xccab8ee7, 11, 0xffffffff, 7, 0x31b6ab00, 15, 0xffe50000, 17, 1812433253>; // mt11213b
+	RandGen mtRand(InRandomSeed);
+
+	const TreeT& tree = gridPtr->tree();
+	openvdb::math::Transform transform = gridPtr->transform();
+
+	std::vector<openvdb::Vec3R> instancePoints;
+	{
+		// Compute a mask of the voxels enclosed by the isosurface.
+		typename openvdb::Grid<BoolTreeT>::Ptr interiorMaskPtr;
+		if (gridPtr->getGridClass() == openvdb::GRID_LEVEL_SET) {
+			interiorMaskPtr = openvdb::tools::sdfInteriorMask(*gridPtr, isovalue);
+		}
+		else {
+			// For non-level-set grids, the interior mask comprises the active voxels.
+			interiorMaskPtr = typename openvdb::Grid<BoolTreeT>::Ptr(openvdb::Grid<BoolTreeT>::create(false));
+			interiorMaskPtr->setTransform(transform.copy());
+			interiorMaskPtr->tree().topologyUnion(tree);
+		}
+
+		if (interrupter && interrupter->wasInterrupted()) return;
+
+		// If the interior mask is small and eroding it results in an empty grid,
+		// use the uneroded mask instead.  (But if the minimum sphere count is zero,
+		// then eroding away the mask is acceptable.)
+		if (!addNarrowBandPoints || (minSphereCount <= 0)) {
+			openvdb::tools::erodeActiveValues(interiorMaskPtr->tree(), 1, openvdb::tools::NN_FACE, openvdb::tools::IGNORE_TILES);
+			openvdb::tools::pruneInactive(interiorMaskPtr->tree());
+		}
+		else {
+			auto& maskTree = interiorMaskPtr->tree();
+			auto copyOfTree = openvdb::StaticPtrCast<BoolTreeT>(maskTree.copy());
+			openvdb::tools::erodeActiveValues(maskTree, 1, openvdb::tools::NN_FACE, openvdb::tools::IGNORE_TILES);
+			openvdb::tools::pruneInactive(maskTree);
+			if (maskTree.empty()) { interiorMaskPtr->setTree(copyOfTree); }
+		}
+
+		// Scatter candidate sphere centroids (instancePoints)
+		instancePoints.reserve(instances);
+		openvdb::tools::v2s_internal::PointAccessor ptnAcc(instancePoints);
+
+		const auto scatterCount = openvdb::Index64(addNarrowBandPoints ? (instances / 2) : instances);
+
+		FRandomStream RandStream(InRandomSeed);
+		float NumPointsPerVoxel = InMinNumberOfPointsPerVoxel;
+		if (InMaxNumberOfPointsPerVoxel > InMinNumberOfPointsPerVoxel)
+		{
+			NumPointsPerVoxel = RandStream.FRandRange(InMinNumberOfPointsPerVoxel, InMaxNumberOfPointsPerVoxel);
+		}
+
+		if (InScatterType == EDataflowVolumeScatterType::Uniform)
+		{
+			openvdb::tools::UniformPointScatter<openvdb::tools::v2s_internal::PointAccessor, RandGen, InterrupterT> UniformScatter(
+				ptnAcc, scatterCount, mtRand, double(InSpread), interrupter);
+			UniformScatter(*interiorMaskPtr);
+		}
+		else if (InScatterType == EDataflowVolumeScatterType::DenseUniform)
+		{
+			openvdb::tools::DenseUniformPointScatter<openvdb::tools::v2s_internal::PointAccessor, RandGen, InterrupterT> DenseUniformScatter(
+				ptnAcc, NumPointsPerVoxel, mtRand, double(InSpread), interrupter);
+			DenseUniformScatter(*interiorMaskPtr);
+		}
+		else if (InScatterType == EDataflowVolumeScatterType::NonUniform)
+		{
+			openvdb::tools::NonUniformPointScatter<openvdb::tools::v2s_internal::PointAccessor, RandGen, InterrupterT> NonUniformScatter(
+				ptnAcc, NumPointsPerVoxel, mtRand, double(InSpread), interrupter);
+			NonUniformScatter(*interiorMaskPtr);
+		}
+	}
+
+	if (interrupter && interrupter->wasInterrupted()) return;
+
+	auto csp = openvdb::tools::ClosestSurfacePoint<GridT>::create(*gridPtr, isovalue, interrupter);
+	if (!csp) return;
+
+	// Add extra instance points in the interior narrow band.
+	if (instancePoints.size() < size_t(instances)) {
+		const Int16TreeT& signTree = csp->signTree();
+		for (auto leafIt = signTree.cbeginLeaf(); leafIt; ++leafIt) {
+			for (auto it = leafIt->cbeginValueOn(); it; ++it) {
+				const int flags = int(it.getValue());
+				if (!(openvdb::tools::volume_to_mesh_internal::EDGES & flags)
+					&& (openvdb::tools::volume_to_mesh_internal::INSIDE & flags))
+				{
+					instancePoints.push_back(transform.indexToWorld(it.getCoord()));
+				}
+				if (instancePoints.size() == size_t(instances)) break;
+			}
+			if (instancePoints.size() == size_t(instances)) break;
+		}
+	}
+
+	if (interrupter && interrupter->wasInterrupted()) return;
+
+	// Assign a radius to each candidate sphere.  The radius is the world-space
+	// distance from the sphere's center to the closest surface point.
+	std::vector<float> instanceRadius;
+	if (!csp->search(instancePoints, instanceRadius)) return;
+
+	minRadius = float(minRadius * transform.voxelSize()[0]);
+	maxRadius = float(maxRadius * transform.voxelSize()[0]);
+
+	float largestRadius = 0.0;
+	int largestRadiusIdx = 0;
+	for (size_t n = 0, N = instancePoints.size(); n < N; ++n) 
+	{
+		if (instanceRadius[n] > maxRadius)
+		{
+			instanceRadius[n] = maxRadius;
+		}
+
+		if (instanceRadius[n] > largestRadius) 
+		{
+			largestRadius = instanceRadius[n];
+			largestRadiusIdx = int(n);
+		}
+	}
+
+	std::vector<unsigned char> instanceMask(instancePoints.size(), 0);
+
+
+	for (size_t s = 0, S = std::min(size_t(maxSphereCount), instancePoints.size()); s < S; ++s) {
+
+		if (interrupter && interrupter->wasInterrupted()) return;
+
+		largestRadius = std::min(maxRadius, largestRadius);
+
+		if ((int(s) >= minSphereCount) && (largestRadius < minRadius)) break;
+
+		const openvdb::Vec4s sphere(
+			float(instancePoints[largestRadiusIdx].x()),
+			float(instancePoints[largestRadiusIdx].y()),
+			float(instancePoints[largestRadiusIdx].z()),
+			largestRadius);
+
+		spheres.push_back(sphere);
+		instanceMask[largestRadiusIdx] = 1;
+
+		openvdb::tools::v2s_internal::UpdatePoints op(
+			sphere, instancePoints, instanceRadius, instanceMask, overlapping);
+		op.run();
+
+		largestRadius = op.radius();
+		largestRadiusIdx = op.index();
+	}
+} // fillWithSpheres
+
+
diff --git a/Engine/Plugins/Experimental/DataflowVolumetricPlugin/Source/DataflowVolumeNodes/Private/Dataflow/DataflowVolumeNodes.cpp b/Engine/Plugins/Experimental/DataflowVolumetricPlugin/Source/DataflowVolumeNodes/Private/Dataflow/DataflowVolumeNodes.cpp
--- a/Engine/Plugins/Experimental/DataflowVolumetricPlugin/Source/DataflowVolumeNodes/Private/Dataflow/DataflowVolumeNodes.cpp
+++ b/Engine/Plugins/Experimental/DataflowVolumetricPlugin/Source/DataflowVolumeNodes/Private/Dataflow/DataflowVolumeNodes.cpp
@@ -31,6 +31,7 @@ namespace UE::Dataflow
 		// Convert
 		DATAFLOW_NODE_REGISTER_CREATION_FACTORY(FCollectionToVolumeDataflowNode);
 		DATAFLOW_NODE_REGISTER_CREATION_FACTORY(FVolumeToSpheresDataflowNode);
+		DATAFLOW_NODE_REGISTER_CREATION_FACTORY(FVolumeToSpheresDataflowNode_v2);
 		DATAFLOW_NODE_REGISTER_CREATION_FACTORY(FConvertVolumeDataflowNode);
 
 		// Scatter
@@ -666,6 +667,204 @@ void FVolumeToSpheresDataflowNode::DebugDraw(UE::Dataflow::FContext& Context, ID
 
 /* --------------------------------------------------------------------------------------------- */
 
+FVolumeToSpheresDataflowNode_v2::FVolumeToSpheresDataflowNode_v2(const UE::Dataflow::FNodeParameters& InParam, FGuid InGuid)
+	: FDataflowNode(InParam, InGuid)
+{
+	RegisterInputConnection(&FloatVolume);
+	RegisterInputConnection(&MinSphereCount).SetCanHidePin(true).SetPinIsHidden(true);
+	RegisterInputConnection(&MaxSphereCount).SetCanHidePin(true).SetPinIsHidden(true);
+	RegisterInputConnection(&MinRadius).SetCanHidePin(true).SetPinIsHidden(true);
+	RegisterInputConnection(&MaxRadius).SetCanHidePin(true).SetPinIsHidden(true);
+	RegisterInputConnection(&IsoValue).SetCanHidePin(true).SetPinIsHidden(true);
+	RegisterInputConnection(&InstanceCount).SetCanHidePin(true).SetPinIsHidden(true);
+	RegisterInputConnection(&RandomSeed).SetCanHidePin(true).SetPinIsHidden(true);
+	RegisterInputConnection(&Spread).SetCanHidePin(true).SetPinIsHidden(true);
+	RegisterInputConnection(&MinNumberOfPointsPerVoxel).SetCanHidePin(true).SetPinIsHidden(true);
+	RegisterInputConnection(&MaxNumberOfPointsPerVoxel).SetCanHidePin(true).SetPinIsHidden(true);
+	RegisterOutputConnection(&Collection);
+	RegisterOutputConnection(&Spheres);
+
+	DebugDrawRenderSettings.RenderType = EDataflowDebugDrawRenderType::Wireframe;
+	DebugDrawRenderSettings.bTranslucent = false;
+	DebugDrawRenderSettings.Color = FLinearColor::Blue;
+	DebugDrawRenderSettings.LineWidthMultiplier = 1.5f;
+}
+
+void FVolumeToSpheresDataflowNode_v2::Evaluate(UE::Dataflow::FContext& Context, const FDataflowOutput* Out) const
+{
+	if (Out->IsA(&Collection) || Out->IsA(&Spheres))
+	{
+		if (IsConnected(&FloatVolume))
+		{
+			const FDataflowFloatVolume& InFloatVolume = GetValue(Context, &FloatVolume);
+			const int32 InMinSphereCount = GetValue(Context, &MinSphereCount);
+			const int32 InMaxSphereCount = GetValue(Context, &MaxSphereCount);
+			const float InMinRadius = GetValue(Context, &MinRadius);
+			const float InMaxRadius = GetValue(Context, &MaxRadius);
+			const float InIsoValue = GetValue(Context, &IsoValue);
+			const int32 InInstanceCount = GetValue(Context, &InstanceCount);
+			const int32 InRandomSeed = GetValue(Context, &RandomSeed);
+			const float InSpread = GetValue(Context, &Spread);
+			const float InMinNumberOfPointsPerVoxel = GetValue(Context, &MinNumberOfPointsPerVoxel);
+			const float InMaxNumberOfPointsPerVoxel = GetValue(Context, &MaxNumberOfPointsPerVoxel);
+
+			TArray<FVector> OutSphereCenters;
+			TArray<float> OutSphereRadii;
+
+			// Compute sphere packing
+			InFloatVolume.VolumeToSpheresImproved(
+				InMinSphereCount,
+				InMaxSphereCount,
+				bOverlapping,
+				InMinRadius,
+				InMaxRadius,
+				InIsoValue,
+				InInstanceCount,
+				InRandomSeed,
+				ScatterType,
+				InSpread,
+				InMinNumberOfPointsPerVoxel,
+				InMaxNumberOfPointsPerVoxel,
+				OutSphereCenters,
+				OutSphereRadii);
+
+			// Add spheres to OutCollection
+			FManagedArrayCollection OutCollection;
+
+			const int32 NumSpheresFromSDF = OutSphereCenters.Num();
+
+			if (NumSpheresFromSDF > 0)
+			{
+				// Add new element to groups
+				static const FName SpheresGroup = "Spheres";
+
+				const int32 NumTransforms = OutCollection.NumElements(FGeometryCollection::TransformGroup);
+				const int32 NumSpheres = OutCollection.NumElements(SpheresGroup);
+
+				OutCollection.AddElements(1, FGeometryCollection::TransformGroup);
+				OutCollection.AddElements(NumSpheresFromSDF, SpheresGroup);
+
+				TManagedArray<FTransform>& Transform = OutCollection.AddAttribute<FTransform>("Transform", FGeometryCollection::TransformGroup);
+				TManagedArray<FString>& BoneName = OutCollection.AddAttribute<FString>("BoneName", FGeometryCollection::TransformGroup);
+				TManagedArray<FLinearColor>& BoneColor = OutCollection.AddAttribute<FLinearColor>("BoneColor", FGeometryCollection::TransformGroup);
+				TManagedArray<int32>& Parent = OutCollection.AddAttribute<int32>("Parent", FGeometryCollection::TransformGroup);
+				TManagedArray<TSet<int32>>& Children = OutCollection.AddAttribute<TSet<int32>>("Children", FGeometryCollection::TransformGroup);
+				TManagedArray<int32>& TransformToSphereIndex = OutCollection.AddAttribute<int32>("TransformToSphereIndex", FGeometryCollection::TransformGroup);
+				TManagedArray<int32>& BoneMap = OutCollection.AddAttribute<int32>("BoneMap", SpheresGroup);
+				TManagedArray<int32>& TransformIndex = OutCollection.AddAttribute<int32>("TransformIndex", SpheresGroup);
+				TManagedArray<FVector3f>& Center = OutCollection.AddAttribute<FVector3f>("Center", SpheresGroup);
+				TManagedArray<float>& Radius = OutCollection.AddAttribute<float>("Radius", SpheresGroup);
+
+				Transform[NumTransforms] = FTransform::Identity;
+				BoneName[NumTransforms] = FString(TEXT("Spheres"));
+				BoneColor[NumTransforms] = FLinearColor(0.02f, 0.01f, 0.1f, 1.0f);
+				Parent[NumTransforms] = -1;
+				Children[NumTransforms] = TSet<int32>();
+				TransformToSphereIndex[NumTransforms] = NumSpheres;
+
+				TArray<FSphere> OutSpheres; OutSpheres.AddUninitialized(NumSpheresFromSDF);
+
+				for (int32 Idx = 0; Idx < NumSpheresFromSDF; ++Idx)
+				{
+					Center[NumSpheres + Idx] = (FVector3f)OutSphereCenters[Idx];
+					Radius[NumSpheres + Idx] = OutSphereRadii[Idx];
+
+					OutSpheres[Idx].Center = OutSphereCenters[Idx];
+					OutSpheres[Idx].W = OutSphereRadii[Idx];
+				}
+
+				TransformIndex[NumSpheres] = NumTransforms;
+
+				SetValue(Context, MoveTemp(OutCollection), &Collection);
+				SetValue(Context, MoveTemp(OutSpheres), &Spheres);
+
+				return;
+			}
+		}
+
+		SetValue(Context, FManagedArrayCollection(), &Collection);
+		SetValue(Context, TArray<FSphere>(), &Spheres);
+	}
+}
+
+#if WITH_EDITOR
+bool FVolumeToSpheresDataflowNode_v2::CanDebugDrawViewMode(const FName& ViewModeName) const
+{
+	return ViewModeName == UE::Dataflow::FDataflowConstruction3DViewMode::Name;
+}
+
+void FVolumeToSpheresDataflowNode_v2::DebugDraw(UE::Dataflow::FContext& Context, IDataflowDebugDrawInterface& DataflowRenderingInterface, const FDebugDrawParameters& DebugDrawParameters) const
+{
+	using namespace UE::Dataflow::Overlay;
+
+	if (DebugDrawParameters.bNodeIsSelected || DebugDrawParameters.bNodeIsPinned)
+	{
+		FString OverlayEvaluateResultString;
+
+		const FDataflowOutput* SpheresOutput = FindOutput(&Spheres);
+		if (SpheresOutput)
+		{
+			const TArray<FSphere>& OutSpheres = SpheresOutput->GetValue(Context, Spheres);
+			//			const FDataflowFloatVolume& OutFloatVolume = this->GetOutputValue(Context, FloatVolume);
+
+			DebugDrawRenderSettings.SetDebugDrawSettings(DataflowRenderingInterface);
+
+			float MinRad = FLT_MAX, MaxRad = -FLT_MAX, SumRad = 0;
+			for (const FSphere& Sphere : OutSpheres)
+			{
+				if (Sphere.W < MinRad)
+				{
+					MinRad = Sphere.W;
+				}
+				if (Sphere.W > MaxRad)
+				{
+					MaxRad = Sphere.W;
+				}
+
+				SumRad += Sphere.W;
+
+				if (bDrawSphereCenters)
+				{
+					DataflowRenderingInterface.DrawPoint(Sphere.Center);
+				}
+
+				if (bDrawSpheres)
+				{
+					DataflowRenderingInterface.DrawSphere(Sphere.Center, Sphere.W);
+				}
+			}
+
+			if (OutSpheres.Num())
+			{
+				const int32 NumSpheres = OutSpheres.Num();
+				const float MeanRad = SumRad / (float)NumSpheres;
+
+				OverlayEvaluateResultString.Appendf(TEXT("Number of Spheres: %d\nMin radius: %3.3f\nMax radius: %3.3f\nMedian radius: %3.3f\n"),
+					NumSpheres,
+					MinRad,
+					MaxRad,
+					MeanRad);
+			}
+			else
+			{
+				OverlayEvaluateResultString.Appendf(TEXT("Number of Spheres: 0\nMin radius: 0.000\nMax radius: 0.000\nMedian radius: 0.000\n"));
+			}
+		}
+
+		// ------------------------------------------------------------------------------------------------------
+		// Overlay
+		TArray<FString> OverlayStrArr;
+		OverlayStrArr.Add(BuildOverlayNodeInfoString(this));
+		OverlayStrArr.Add(OverlayEvaluateResultString);
+		FString OverlayStr = BuildOverlayFinalString(OverlayStrArr);
+		DataflowRenderingInterface.DrawOverlayText(OverlayStr);
+		// ------------------------------------------------------------------------------------------------------
+	}
+}
+#endif
+
+/* --------------------------------------------------------------------------------------------- */
+
 FConvertVolumeDataflowNode::FConvertVolumeDataflowNode(const UE::Dataflow::FNodeParameters& InParam, FGuid InGuid)
 	: FDataflowNode(InParam, InGuid)
 {
diff --git a/Engine/Plugins/Experimental/DataflowVolumetricPlugin/Source/DataflowVolumeNodes/Private/Dataflow/DataflowVolumeNodes.h b/Engine/Plugins/Experimental/DataflowVolumetricPlugin/Source/DataflowVolumeNodes/Private/Dataflow/DataflowVolumeNodes.h
--- a/Engine/Plugins/Experimental/DataflowVolumetricPlugin/Source/DataflowVolumeNodes/Private/Dataflow/DataflowVolumeNodes.h
+++ b/Engine/Plugins/Experimental/DataflowVolumetricPlugin/Source/DataflowVolumeNodes/Private/Dataflow/DataflowVolumeNodes.h
@@ -10,6 +10,8 @@
 #include "Dataflow/DataflowVolumeAnyType.h"
 #include "Dataflow/DataflowDebugDraw.h"
 #include "Dataflow/DataflowColorRamp.h"
+#include "Dataflow/DataflowNode.h"
+#include "Dataflow/DataflowTerminalNode.h"
 
 #include "DataflowVolumeNodes.generated.h"
 
@@ -293,7 +295,7 @@ public:
 
 /**
  *
- * Fill an SDF/Fog volume with spheres
+ * Fill an SDF/Fog volume with spheres, using the original algorithm
  *
  */
 USTRUCT(meta = (Experimental, DataflowVolume))
@@ -360,6 +362,93 @@ public:
 	FVolumeToSpheresDataflowNode(const UE::Dataflow::FNodeParameters& InParam, FGuid InGuid = FGuid::NewGuid());
 };
 
+/**
+ *
+ * Fill an SDF/Fog volume with spheres, using the improved algorithm
+ *
+ */
+USTRUCT(meta = (Experimental, DataflowVolume))
+struct FVolumeToSpheresDataflowNode_v2 : public FDataflowNode
+{
+	GENERATED_USTRUCT_BODY()
+	DATAFLOW_NODE_DEFINE_INTERNAL(FVolumeToSpheresDataflowNode_v2, "VolumeToSpheres", "Volume|Utilities", "")
+
+private:
+	/** Float volume */
+	UPROPERTY(meta = (DataflowInput))
+	FDataflowFloatVolume FloatVolume;
+
+	UPROPERTY(EditAnywhere, Category = "Spheres", meta = (DataflowInput, DisplayName = "Minimum Sphere Count"));
+	int32 MinSphereCount = 1;
+
+	UPROPERTY(EditAnywhere, Category = "Spheres", meta = (DataflowInput, DisplayName = "Maximum Sphere Count"));
+	int32 MaxSphereCount = 50;
+
+	UPROPERTY(EditAnywhere, Category = "Spheres");
+	bool bOverlapping = false;
+
+	UPROPERTY(EditAnywhere, Category = "Spheres", meta = (DataflowInput, DisplayName = "Minimum Radius", ClampMin = 0.001f));
+	float MinRadius = 1.f;
+
+	UPROPERTY(EditAnywhere, Category = "Spheres", meta = (DataflowInput, DisplayName = "Maximum Radius", ClampMin = 0.001f));
+	float MaxRadius = 1000.f;
+
+	UPROPERTY(EditAnywhere, Category = "Spheres", meta = (DataflowInput));
+	float IsoValue = 0.f;
+
+	UPROPERTY(EditAnywhere, Category = "Spheres", meta = (DataflowInput, ClampMin = 1));
+	int32 InstanceCount = 10000;
+
+	UPROPERTY(EditAnywhere, Category = "Scatter")
+	EDataflowVolumeScatterType ScatterType = EDataflowVolumeScatterType::Uniform;
+
+	UPROPERTY(EditAnywhere, Category = "Scatter", meta = (DataflowInput));
+	int32 RandomSeed = 0;
+
+	/** */
+	UPROPERTY(EditAnywhere, Category = "Scatter", meta = (DataflowInput, ClampMin = 0, ClampMax = 1));
+	float Spread = 1.f;
+
+	/** Minimum for the random range */
+	UPROPERTY(EditAnywhere, Category = "Scatter", meta = (DataflowInput, ClampMin = 0, EditCondition = "ScatterType != EDataflowVolumeScatterType::Uniform", EditConditionHides));
+	float MinNumberOfPointsPerVoxel = 1.f;
+
+	/** Maximum for the random range */
+	UPROPERTY(EditAnywhere, Category = "Scatter", meta = (DataflowInput, ClampMin = 0, EditCondition = "ScatterType != EDataflowVolumeScatterType::Uniform", EditConditionHides));
+	float MaxNumberOfPointsPerVoxel = 1.f;
+
+	UPROPERTY(EditAnywhere, Category = "Debug Draw", meta = (ShowOnlyInnerProperties))
+	FDataflowNodeDebugDrawSettings DebugDrawRenderSettings;
+
+	UPROPERTY(EditAnywhere, Category = "Debug Draw")
+	bool bDrawSpheres = true;
+
+	UPROPERTY(EditAnywhere, Category = "Debug Draw")
+	bool bDrawSphereCenters = false;
+
+	/**  */
+	UPROPERTY(meta = (DataflowOutput))
+	FManagedArrayCollection Collection;
+
+	/**  */
+	UPROPERTY(meta = (DataflowOutput))
+	TArray<FSphere> Spheres;
+
+	virtual void Evaluate(UE::Dataflow::FContext& Context, const FDataflowOutput* Out) const override;
+
+#if WITH_EDITOR
+	virtual bool CanDebugDraw() const override
+	{
+		return true;
+	}
+	virtual bool CanDebugDrawViewMode(const FName& ViewModeName) const override;
+	virtual void DebugDraw(UE::Dataflow::FContext& Context, IDataflowDebugDrawInterface& DataflowRenderingInterface, const FDebugDrawParameters& DebugDrawParameters) const override;
+#endif
+
+public:
+	FVolumeToSpheresDataflowNode_v2(const UE::Dataflow::FNodeParameters& InParam, FGuid InGuid = FGuid::NewGuid());
+};
+
 /**
  * Convert a volume to Volume/Polygons
  */
