Dataflow : Deprecate dataflow specific color ramp types and replace it with the engine one

#Synced-CL 48963730

diff --git a/Engine/Plugins/Experimental/Dataflow/Source/DataflowEditor/Private/Dataflow/DataflowColorRampCustomization.cpp b/Engine/Plugins/Experimental/Dataflow/Source/DataflowEditor/Private/Dataflow/DataflowColorRampCustomization.cpp
--- a/Engine/Plugins/Experimental/Dataflow/Source/DataflowEditor/Private/Dataflow/DataflowColorRampCustomization.cpp
+++ b/Engine/Plugins/Experimental/Dataflow/Source/DataflowEditor/Private/Dataflow/DataflowColorRampCustomization.cpp
@@ -8,6 +8,7 @@
 #include "Dataflow/DataflowGraphEditor.h"
 #include "DetailWidgetRow.h"
 
+
 namespace UE::Dataflow
 {
 	TSharedRef<IPropertyTypeCustomization> FDataflowColorRampCustomization::MakeInstance()
@@ -17,12 +18,12 @@ namespace UE::Dataflow
 
 	void FDataflowColorRampCustomization::CustomizeHeader(TSharedRef<IPropertyHandle> PropertyHandle, FDetailWidgetRow& HeaderRow, IPropertyTypeCustomizationUtils& CustomizationUtils)
 	{
+		PRAGMA_DISABLE_DEPRECATION_WARNINGS
 		if (const FStructProperty* const StructProperty = CastField<FStructProperty>(PropertyHandle->GetProperty()))
 		{
-			if (StructProperty->Struct)
+			if (StructProperty->Struct == FDataflowColorRamp::StaticStruct())
 			{
 				void* Data = nullptr;
-
 				if (PropertyHandle->GetValueData(Data) == FPropertyAccess::Success)
 				{
 					FDataflowColorRamp* ColorRamp = reinterpret_cast<FDataflowColorRamp*>(Data);
@@ -41,6 +42,7 @@ namespace UE::Dataflow
 				}
 			}
 		}
+		PRAGMA_ENABLE_DEPRECATION_WARNINGS
 	}
 
 	TSharedRef<IPropertyTypeCustomization> FDataflowLinearColorRampCustomization::MakeInstance()
@@ -50,6 +52,7 @@ namespace UE::Dataflow
 
 	void FDataflowLinearColorRampCustomization::CustomizeHeader(TSharedRef<IPropertyHandle> PropertyHandle, FDetailWidgetRow& HeaderRow, IPropertyTypeCustomizationUtils& CustomizationUtils)
 	{
+		PRAGMA_DISABLE_DEPRECATION_WARNINGS
 		if (const FStructProperty* const StructProperty = CastField<FStructProperty>(PropertyHandle->GetProperty()))
 		{
 			if (StructProperty->Struct)
@@ -74,8 +77,11 @@ namespace UE::Dataflow
 				}
 			}
 		}
+		PRAGMA_ENABLE_DEPRECATION_WARNINGS
 	}
 }  
 
 
 
+
+
diff --git a/Engine/Plugins/Experimental/Dataflow/Source/DataflowEditor/Private/DataflowEditorTools/DataflowEditorVertexAttributePaintToolDetailCustomization.cpp b/Engine/Plugins/Experimental/Dataflow/Source/DataflowEditor/Private/DataflowEditorTools/DataflowEditorVertexAttributePaintToolDetailCustomization.cpp
--- a/Engine/Plugins/Experimental/Dataflow/Source/DataflowEditor/Private/DataflowEditorTools/DataflowEditorVertexAttributePaintToolDetailCustomization.cpp
+++ b/Engine/Plugins/Experimental/Dataflow/Source/DataflowEditor/Private/DataflowEditorTools/DataflowEditorVertexAttributePaintToolDetailCustomization.cpp
@@ -325,7 +325,7 @@ namespace UE::Dataflow::Editor
 		auto MakeColorRampWidget =
 			[](
 				EDataflowEditorToolColorMode ColorMode, 
-				FDataflowLinearColorRamp* ColorRamp,
+				FLinearColorRamp* ColorRamp,
 				TWeakObjectPtr<UDataflowEditorVertexAttributePaintToolProperties> ToolProperties
 				)
 			{
@@ -749,7 +749,7 @@ namespace UE::Dataflow::Editor
 
 	void FVertexAttributePaintToolDetailCustomization::AddBrushHitBackFacesRow(IDetailCategoryBuilder& BrushCategory) const
 	{
-		BrushCategory.AddCustomRow(LOCTEXT("BrushHitBackFacesCategory", "Brush Falloff"), false)
+		BrushCategory.AddCustomRow(LOCTEXT("BrushHitBackFacesCategory", "Brush Hit Back Faces"), false)
 			.NameContent()
 			[
 				SNew(STextBlock)
diff --git a/Engine/Plugins/Experimental/Dataflow/Source/DataflowEditor/Public/DataflowEditorTools/DataflowEditorVertexAttributePaintTool.h b/Engine/Plugins/Experimental/Dataflow/Source/DataflowEditor/Public/DataflowEditorTools/DataflowEditorVertexAttributePaintTool.h
--- a/Engine/Plugins/Experimental/Dataflow/Source/DataflowEditor/Public/DataflowEditorTools/DataflowEditorVertexAttributePaintTool.h
+++ b/Engine/Plugins/Experimental/Dataflow/Source/DataflowEditor/Public/DataflowEditorTools/DataflowEditorVertexAttributePaintTool.h
@@ -24,7 +24,7 @@
 #include "GroupTopology.h"
 #include "MeshVertexPaintTool.h"
 #include "DataflowEditorTools/DataflowEditorToolBuilder.h"
-#include "Dataflow/DataflowColorRamp.h"
+#include "Curves/LinearColorRamp.h"
 #include "DataflowEditorTools/DataflowEditorToolEnums.h"
 #include "ToolMeshSelector.h"
 
@@ -183,14 +183,14 @@ struct FDataflowEditorVertexAttributePaintToolDisplayProperties
 
 	// Color ramp to use when the color mode is set to "Ramp"
 	UPROPERTY()
-	FDataflowLinearColorRamp ColorRamp;
+	FLinearColorRamp ColorRamp;
 
 	UPROPERTY()
-	TArray<FDataflowLinearColorRamp> ColorRampPresets;
+	TArray<FLinearColorRamp> ColorRampPresets;
 
 	// Used by UI customization 
-	FDataflowLinearColorRamp GreyScaleColorRamp;
-	FDataflowLinearColorRamp WhiteColorRamp;
+	FLinearColorRamp GreyScaleColorRamp;
+	FLinearColorRamp WhiteColorRamp;
 };
 
 USTRUCT()
diff --git a/Engine/Plugins/Experimental/DataflowVolumetricPlugin/Source/DataflowVolumeNodes/Private/Dataflow/DataflowVolumeNodes.cpp b/Engine/Plugins/Experimental/DataflowVolumetricPlugin/Source/DataflowVolumeNodes/Private/Dataflow/DataflowVolumeNodes.cpp
--- a/Engine/Plugins/Experimental/DataflowVolumetricPlugin/Source/DataflowVolumeNodes/Private/Dataflow/DataflowVolumeNodes.cpp
+++ b/Engine/Plugins/Experimental/DataflowVolumetricPlugin/Source/DataflowVolumeNodes/Private/Dataflow/DataflowVolumeNodes.cpp
@@ -1657,14 +1657,14 @@ FVolumeSliceDataflowNode::FVolumeSliceDataflowNode(const UE::Dataflow::FNodePara
 
 	for (int32 Idx = 0; Idx < NumColors; ++Idx)
 	{
-		ColorRamp.ColorRamp.GetCurves()[0].CurveToEdit->AddKey((float)Idx * 1.f / ((float)NumColors - 1.f), Spectrum[Idx].R);
-		ColorRamp.ColorRamp.GetCurves()[1].CurveToEdit->AddKey((float)Idx * 1.f / ((float)NumColors - 1.f), Spectrum[Idx].G);
-		ColorRamp.ColorRamp.GetCurves()[2].CurveToEdit->AddKey((float)Idx * 1.f / ((float)NumColors - 1.f), Spectrum[Idx].B);
+		LinearColorRamp.GetCurves()[0].CurveToEdit->AddKey((float)Idx * 1.f / ((float)NumColors - 1.f), Spectrum[Idx].R);
+		LinearColorRamp.GetCurves()[1].CurveToEdit->AddKey((float)Idx * 1.f / ((float)NumColors - 1.f), Spectrum[Idx].G);
+		LinearColorRamp.GetCurves()[2].CurveToEdit->AddKey((float)Idx * 1.f / ((float)NumColors - 1.f), Spectrum[Idx].B);
 	}
-	ColorRamp.ColorRamp.GetCurves()[3].CurveToEdit->AddKey(0.f, 1.f);
+	LinearColorRamp.GetCurves()[3].CurveToEdit->AddKey(0.f, 1.f);
 
 	//
-	ColorRamp.ColorRamp.OnColorCurveChangedDelegate.AddRaw(this, &FVolumeSliceDataflowNode::OnColorCurveChanged);
+	LinearColorRamp.OnColorCurveChangedDelegate.AddRaw(this, &FVolumeSliceDataflowNode::OnColorCurveChanged);
 }
 
 void FVolumeSliceDataflowNode::OnColorCurveChanged(TArray<FRichCurve*> Curves)
@@ -1741,16 +1741,16 @@ void FVolumeSliceDataflowNode::OnPropertyChanged(UE::Dataflow::FContext& Context
 		// Update ColorRamp
 		for (int32 Idx = 0; Idx < 4; ++Idx)
 		{
-			ColorRamp.ColorRamp.GetCurves()[Idx].CurveToEdit->Reset();
+			LinearColorRamp.GetCurves()[Idx].CurveToEdit->Reset();
 		}
 
 		for (int32 Idx = 0; Idx < NumColors; ++Idx)
 		{
-			ColorRamp.ColorRamp.GetCurves()[0].CurveToEdit->AddKey((float)Idx * 1.f / ((float)NumColors - 1.f), Spectrum[Idx].R);
-			ColorRamp.ColorRamp.GetCurves()[1].CurveToEdit->AddKey((float)Idx * 1.f / ((float)NumColors - 1.f), Spectrum[Idx].G);
-			ColorRamp.ColorRamp.GetCurves()[2].CurveToEdit->AddKey((float)Idx * 1.f / ((float)NumColors - 1.f), Spectrum[Idx].B);
+			LinearColorRamp.GetCurves()[0].CurveToEdit->AddKey((float)Idx * 1.f / ((float)NumColors - 1.f), Spectrum[Idx].R);
+			LinearColorRamp.GetCurves()[1].CurveToEdit->AddKey((float)Idx * 1.f / ((float)NumColors - 1.f), Spectrum[Idx].G);
+			LinearColorRamp.GetCurves()[2].CurveToEdit->AddKey((float)Idx * 1.f / ((float)NumColors - 1.f), Spectrum[Idx].B);
 		}
-		ColorRamp.ColorRamp.GetCurves()[3].CurveToEdit->AddKey(0.f, 1.f);
+		LinearColorRamp.GetCurves()[3].CurveToEdit->AddKey(0.f, 1.f);
 
 	}
 #endif
@@ -1795,87 +1795,18 @@ void FVolumeSliceDataflowNode::DebugDraw(UE::Dataflow::FContext& Context, IDataf
 					DataflowRenderingInterface.SetColor(FLinearColor::Gray);
 					DataflowRenderingInterface.DrawBox(0.5 * Extent, FQuat::Identity, VolumeBoundingBox.GetCenter(), 1.f);
 
-					float Progress;
-
 					for (int32 Idx = 0; Idx < SamplePoints.Num(); ++Idx)
 					{
-						float SampledValue = FMath::Clamp(SampledValues[Idx], Range.X, Range.Y);
-
-						FLinearColor Color;
-						if (ColorRampType == EDataflowVolumeSliceRamp::WhiteToRed)
-						{
-							Progress = 1.f - (Range.Y - SampledValue) / (Range.Y - Range.X);
-//							Color = FLinearColor::LerpUsingHSV(FLinearColor::White, FLinearColor::Red, Progress);
-//							Color = ColorRamp.ColorRamp.GetLinearColorValue(Progress);
-						}
-						else if (ColorRampType == EDataflowVolumeSliceRamp::Grayscale)
-						{
-							Progress = 1.f - (Range.Y - SampledValue) / (Range.Y - Range.X);
-//							Color = FLinearColor::LerpUsingHSV(FLinearColor::Black, FLinearColor::White, Progress);
-//							Color = ColorRamp.ColorRamp.GetLinearColorValue(Progress);
-						}
-						else if (ColorRampType == EDataflowVolumeSliceRamp::Custom)
-						{
-							Progress = 1.f - (Range.Y - SampledValue) / (Range.Y - Range.X);
-//							Color = FLinearColor::LerpUsingHSV(ColorA, ColorB, Progress);
-//							Color = ColorRamp.ColorRamp.GetLinearColorValue(Progress);
-						}
-						else if (ColorRampType == EDataflowVolumeSliceRamp::InfraRed)
-						{
-							const int32 NumColors = 7;
-							//TArray<FLinearColor> Spectrum; Spectrum.AddUninitialized(NumColors);
-
-							//Spectrum[6] = FLinearColor::Red;
-							//Spectrum[5] = FLinearColor(FColor(255, 128, 0, 1)); // Orange
-							//Spectrum[4] = FLinearColor::Yellow;
-							//Spectrum[3] = FLinearColor::Green;
-							//Spectrum[2] = FLinearColor::Blue;
-							//Spectrum[1] = FLinearColor(FColor(64, 0, 128, 1)); // Indigo
-							//Spectrum[0] = FLinearColor(FColor(128, 0, 128, 1)); // Violet
-
-							//if ((SampledValue - Range.X) < 1e-2)
-							//{
-							//	SampledValue += 1e-2;
-							//}
-							//else if ((Range.Y - SampledValue) < 1e-2)
-							//{
-							//	SampledValue -= 1e-2;
-							//}
-
-							//const float Delta = (Range.Y - Range.X) / float(NumColors - 1);
-							//const float B = floorf((SampledValue - Range.X) / Delta);
-							//const float NewMin = Range.X + B * Delta;
-							//const float NewMax = Range.X + (B + 1) * Delta;
-
-//							Progress = 1.f - (NewMax - SampledValue) / Delta;
-							Progress = 1.f - (Range.Y - SampledValue) / (Range.Y - Range.X);
-//							Color = FLinearColor::LerpUsingHSV(Spectrum[int32(B)], Spectrum[int32(B + 1)], Progress);
-						}
-						else if (ColorRampType == EDataflowVolumeSliceRamp::BlackBody)
-						{
-							const int32 NumColors = 4;
-							//TArray<FLinearColor> Spectrum; Spectrum.AddUninitialized(NumColors);
-
-							//Spectrum[3] = FLinearColor::White;
-							//Spectrum[2] = FLinearColor::Yellow;
-							//Spectrum[1] = FLinearColor::Red;
-							//Spectrum[0] = FLinearColor::Black;
-
-							//const float Delta = (Range.Y - Range.X) / float(NumColors - 1);
-							//const float B = floorf((SampledValue - Range.X) / Delta);
-							//const float NewMin = Range.X + B * Delta;
-							//const float NewMax = Range.X + (B + 1) * Delta;
-
-//							Progress = 1.f - (NewMax - SampledValue) / Delta;
-							Progress = 1.f - (Range.Y - SampledValue) / (Range.Y - Range.X);
-							//Color = Spectrum[int32(B)];
-						}
-						Progress = 1.f - (Range.Y - SampledValue) / (Range.Y - Range.X);
+						const float SampledValue = FMath::Clamp(SampledValues[Idx], Range.X, Range.Y);
+						
+						float Progress = ((Range.Y - Range.X) != 0)
+							? (1.f - (Range.Y - SampledValue) / (Range.Y - Range.X))
+							: 0.0f;
 						if (bInvertColorRamp)
 						{
 							Progress = 1.f - Progress;
 						}
-						Color = ColorRamp.ColorRamp.GetLinearColorValue(Progress);
+						const FLinearColor Color = LinearColorRamp.GetLinearColorValue(Progress);
 
 						DataflowRenderingInterface.SetColor(Color);
 						DataflowRenderingInterface.DrawPoint(SamplePoints[Idx]);
diff --git a/Engine/Plugins/Experimental/DataflowVolumetricPlugin/Source/DataflowVolumeNodes/Private/Dataflow/DataflowVolumeNodes.h b/Engine/Plugins/Experimental/DataflowVolumetricPlugin/Source/DataflowVolumeNodes/Private/Dataflow/DataflowVolumeNodes.h
--- a/Engine/Plugins/Experimental/DataflowVolumetricPlugin/Source/DataflowVolumeNodes/Private/Dataflow/DataflowVolumeNodes.h
+++ b/Engine/Plugins/Experimental/DataflowVolumetricPlugin/Source/DataflowVolumeNodes/Private/Dataflow/DataflowVolumeNodes.h
@@ -12,6 +12,7 @@
 #include "Dataflow/DataflowColorRamp.h"
 #include "Dataflow/DataflowNode.h"
 #include "Dataflow/DataflowTerminalNode.h"
+#include "Curves/LinearColorRamp.h"
 
 #include "DataflowVolumeNodes.generated.h"
 
@@ -786,8 +787,8 @@ private:
 	UPROPERTY(EditAnywhere, Category = "Visualization Settings - Scalar Visualization Settings", meta = (EditCondition = "bVisualize == true", EditConditionHides));
 	EDataflowVolumeSliceRamp ColorRampType = EDataflowVolumeSliceRamp::InfraRed;
 
-	UPROPERTY(EditAnywhere, Category = "Visualization Settings - Scalar Visualization Settings", meta = (EditCondition = "bVisualize == true", EditConditionHides));
-	FDataflowColorRamp ColorRamp;
+	UPROPERTY(EditAnywhere, Category = "Visualization Settings - Scalar Visualization Settings", meta = (EditCondition = "bVisualize == true", EditConditionHides, DisplayName = "ColorRamp"));
+	FLinearColorRamp LinearColorRamp;
 
 	UPROPERTY(EditAnywhere, Category = "Visualization Settings - Scalar Visualization Settings");
 	bool bInvertColorRamp = false;
diff --git a/Engine/Source/Runtime/Experimental/Dataflow/Engine/Private/Dataflow/DataflowColorRamp.cpp b/Engine/Source/Runtime/Experimental/Dataflow/Engine/Private/Dataflow/DataflowColorRamp.cpp
--- a/Engine/Source/Runtime/Experimental/Dataflow/Engine/Private/Dataflow/DataflowColorRamp.cpp
+++ b/Engine/Source/Runtime/Experimental/Dataflow/Engine/Private/Dataflow/DataflowColorRamp.cpp
@@ -125,6 +125,7 @@ FLinearColor FDataflowColorCurveOwner::GetCurveColor(FRichCurveEditInfo CurveInf
 
 ////////////////////////////////////////////////////////////////////////////////
 
+PRAGMA_DISABLE_DEPRECATION_WARNINGS
 void FDataflowLinearColorRamp::SetFrom(const FDataflowLinearColorRamp& Other)
 {
 	RichCurves[0] = Other.RichCurves[0];
@@ -133,7 +134,7 @@ void FDataflowLinearColorRamp::SetFrom(const FDataflowLinearColorRamp& Other)
 	RichCurves[3] = Other.RichCurves[3];
 	OnCurveChanged(GetCurves());
 }
-
+PRAGMA_ENABLE_DEPRECATION_WARNINGS
 
 bool FDataflowLinearColorRamp::IsEmpty() const
 {
diff --git a/Engine/Source/Runtime/Experimental/Dataflow/Engine/Public/Dataflow/DataflowColorRamp.h b/Engine/Source/Runtime/Experimental/Dataflow/Engine/Public/Dataflow/DataflowColorRamp.h
--- a/Engine/Source/Runtime/Experimental/Dataflow/Engine/Public/Dataflow/DataflowColorRamp.h
+++ b/Engine/Source/Runtime/Experimental/Dataflow/Engine/Public/Dataflow/DataflowColorRamp.h
@@ -9,7 +9,8 @@
 
 #define UE_API DATAFLOWENGINE_API
 
-struct FDataflowColorCurveOwner : public FCurveOwnerInterface
+
+struct UE_DEPRECATED(5.8, "FDataflowColorCurveOwner is deprecated, use the FLinearColorRamp instead") FDataflowColorCurveOwner : public FCurveOwnerInterface
 {
 public:
 	DECLARE_MULTICAST_DELEGATE_OneParam(FOnColorCurveChanged, TArray<FRichCurve*> /* ChangedCurves */);
@@ -41,15 +42,17 @@ private:
 };
 
 USTRUCT()
-struct FDataflowColorRamp
+struct UE_DEPRECATED(5.8, "FDataflowColorRamp is deprecated, use the FLinearColorRamp instead") FDataflowColorRamp
 {
 	GENERATED_USTRUCT_BODY()
 
+PRAGMA_DISABLE_DEPRECATION_WARNINGS
 	FDataflowColorCurveOwner ColorRamp;
+PRAGMA_ENABLE_DEPRECATION_WARNINGS
 };
 
 USTRUCT()
-struct FDataflowLinearColorRamp 
+struct UE_DEPRECATED(5.8, "FDataflowLinearColorRamp is deprecated, use the FLinearColorRamp instead") FDataflowLinearColorRamp
 #if CPP
 	: public FCurveOwnerInterface
 #endif
