Dataflow:
- Added FVolumeToSpheresWithRelaxationDataflowNode node

#Synced-CL 49790744

diff --git a/Engine/Plugins/Experimental/DataflowVolumetricPlugin/Source/DataflowVolumeCore/Internal/Dataflow/DataflowFloatVolume.h b/Engine/Plugins/Experimental/DataflowVolumetricPlugin/Source/DataflowVolumeCore/Internal/Dataflow/DataflowFloatVolume.h
--- a/Engine/Plugins/Experimental/DataflowVolumetricPlugin/Source/DataflowVolumeCore/Internal/Dataflow/DataflowFloatVolume.h
+++ b/Engine/Plugins/Experimental/DataflowVolumetricPlugin/Source/DataflowVolumeCore/Internal/Dataflow/DataflowFloatVolume.h
@@ -159,6 +159,23 @@ public:
 		const bool InPruneDegenerateTiles,
 		const float InPruneTol) const;
 
+	/** Compute sphere packing in an SDF volume with relaxation */
+	DATAFLOWVOLUMECORE_API void VolumeToSpheresWithRelaxation(
+		const int32 InPointCount,
+		const int32 InScatterRandomSeed,
+		const float InMinRadius,
+		const float InMaxRadius,
+		const int32 InRandomSeed,
+		const float InIsoValue,
+		const float InSpread,
+		const bool InRemoveSpheres,
+		const bool InApplyRelaxation,
+		const int32 InSteps,
+		const float InStepScalar,
+		const bool InCorrectAgainstSurface,
+		const float InDistanceThreshold,
+		TArray<FSphere>& OutSpheres) const;
+
 private:
 	friend struct FDataflowFloatVectorVolume;
 	friend class UE::DataflowVolume::Private::FDataflowFloatVolumeImpl;
diff --git a/Engine/Plugins/Experimental/DataflowVolumetricPlugin/Source/DataflowVolumeCore/Private/Dataflow/DataflowFloatVolume.cpp b/Engine/Plugins/Experimental/DataflowVolumetricPlugin/Source/DataflowVolumeCore/Private/Dataflow/DataflowFloatVolume.cpp
--- a/Engine/Plugins/Experimental/DataflowVolumetricPlugin/Source/DataflowVolumeCore/Private/Dataflow/DataflowFloatVolume.cpp
+++ b/Engine/Plugins/Experimental/DataflowVolumetricPlugin/Source/DataflowVolumeCore/Private/Dataflow/DataflowFloatVolume.cpp
@@ -460,3 +460,48 @@ FDataflowFloatVolume FDataflowFloatVolume::VolumeCombine(
 	return FDataflowFloatVolume();
 }
 
+/* --------------------------------------------------------------------------------------------------------------- */
+
+void FDataflowFloatVolume::VolumeToSpheresWithRelaxation(
+	const int32 InPointCount,
+	const int32 InScatterRandomSeed,
+	const float InMinRadius,
+	const float InMaxRadius,
+	const int32 InRandomSeed,
+	const float InIsoValue,
+	const float InSpread,
+	const bool InRemoveSpheres,
+	const bool InApplyRelaxation,
+	const int32 InSteps,
+	const float InStepScalar,
+	const bool InCorrectAgainstSurface,
+	const float InDistanceThreshold,
+	TArray<FSphere>& OutSpheres) const
+{
+	using namespace UE::DataflowVolume::Private;
+
+	if (TSharedPtr<const FDataflowFloatVolumeImpl> FloatVolumeImpl = GetFloatVolumeImpl())
+	{
+		FloatVolumeImpl->VolumeToSpheresWithRelaxation(
+			InPointCount,
+			InScatterRandomSeed,
+			InMinRadius,
+			InMaxRadius,
+			InRandomSeed,
+			InIsoValue,
+			InSpread,
+			InRemoveSpheres,
+			InApplyRelaxation,
+			InSteps,
+			InStepScalar,
+			InCorrectAgainstSurface,
+			InDistanceThreshold,
+			OutSpheres);
+	}
+}
+
+
+
+
+
+
diff --git a/Engine/Plugins/Experimental/DataflowVolumetricPlugin/Source/DataflowVolumeCore/Private/Dataflow/DataflowFloatVolumeImpl.cpp b/Engine/Plugins/Experimental/DataflowVolumetricPlugin/Source/DataflowVolumeCore/Private/Dataflow/DataflowFloatVolumeImpl.cpp
--- a/Engine/Plugins/Experimental/DataflowVolumetricPlugin/Source/DataflowVolumeCore/Private/Dataflow/DataflowFloatVolumeImpl.cpp
+++ b/Engine/Plugins/Experimental/DataflowVolumetricPlugin/Source/DataflowVolumeCore/Private/Dataflow/DataflowFloatVolumeImpl.cpp
@@ -1287,5 +1287,43 @@ namespace UE::DataflowVolume::Private
 
 		return false;
 	}
+
+	/* --------------------------------------------------------------------------------------------------------------- */
+
+	void FDataflowFloatVolumeImpl::VolumeToSpheresWithRelaxation(
+		const int32 InPointCount,
+		const int32 InScatterRandomSeed,
+		const float InMinRadius,
+		const float InMaxRadius,
+		const int32 InRandomSeed,
+		const float InIsoValue,
+		const float InSpread,
+		const bool InRemoveSpheres,
+		const bool InApplyRelaxation,
+		const int32 InSteps,
+		const float InStepScalar,
+		const bool InCorrectAgainstSurface,
+		const float InDistanceThreshold,
+		TArray<FSphere>& OutSpheres) const
+	{
+		if (FloatGrid)
+		{
+			openvdb::FloatGrid::ConstPtr GridPtr = openvdb::gridConstPtrCast<openvdb::FloatGrid>(FloatGrid);
+
+			std::vector<openvdb::Vec4s> Spheres;
+
+			VolumeToSpheresWithRelaxationProc(*GridPtr,
+				Spheres, InPointCount, InScatterRandomSeed, InMinRadius, InMaxRadius, InRandomSeed,
+				InIsoValue, InSpread, InRemoveSpheres, InApplyRelaxation, InSteps, InStepScalar, InCorrectAgainstSurface, InDistanceThreshold);
+
+			const int32 NumSpheres = Spheres.size();
+			OutSpheres.SetNum(NumSpheres);
+
+			for (int32 Idx = 0; Idx < NumSpheres; ++Idx)
+			{
+				OutSpheres[Idx] = FSphere(FVector(Spheres[Idx].x(), Spheres[Idx].y(), Spheres[Idx].z()), Spheres[Idx].w());
+			}
+		}
+	}
 }
 
diff --git a/Engine/Plugins/Experimental/DataflowVolumetricPlugin/Source/DataflowVolumeCore/Private/Dataflow/DataflowFloatVolumeImpl.h b/Engine/Plugins/Experimental/DataflowVolumetricPlugin/Source/DataflowVolumeCore/Private/Dataflow/DataflowFloatVolumeImpl.h
--- a/Engine/Plugins/Experimental/DataflowVolumetricPlugin/Source/DataflowVolumeCore/Private/Dataflow/DataflowFloatVolumeImpl.h
+++ b/Engine/Plugins/Experimental/DataflowVolumetricPlugin/Source/DataflowVolumeCore/Private/Dataflow/DataflowFloatVolumeImpl.h
@@ -157,6 +157,22 @@ namespace UE::DataflowVolume::Private
 			const EDataflowVolumeSDFCombineInterpolation InInterpolation,
 			const bool InPruneDegenerateTiles,
 			const float InPruneTol) const;
+
+		void VolumeToSpheresWithRelaxation(
+			const int32 InPointCount,
+			const int32 InScatterRandomSeed,
+			const float InMinRadius,
+			const float InMaxRadius,
+			const int32 InRandomSeed,
+			const float InIsoValue,
+			const float InSpread,
+			const bool InRemoveSpheres,
+			const bool InApplyRelaxation,
+			const int32 InSteps,
+			const float InStepScalar,
+			const bool InCorrectAgainstSurface,
+			const float InDistanceThreshold,
+			TArray<FSphere>& OutSpheres) const;
 	};
 }
 
diff --git a/Engine/Plugins/Experimental/DataflowVolumetricPlugin/Source/DataflowVolumeCore/Private/Dataflow/DataflowVolumeAlgo.h b/Engine/Plugins/Experimental/DataflowVolumetricPlugin/Source/DataflowVolumeCore/Private/Dataflow/DataflowVolumeAlgo.h
--- a/Engine/Plugins/Experimental/DataflowVolumetricPlugin/Source/DataflowVolumeCore/Private/Dataflow/DataflowVolumeAlgo.h
+++ b/Engine/Plugins/Experimental/DataflowVolumetricPlugin/Source/DataflowVolumeCore/Private/Dataflow/DataflowVolumeAlgo.h
@@ -258,4 +258,394 @@ void FillWithSpheres(
 	}
 } // fillWithSpheres
 
+/* --------------------------------------------------------------------------------------------------------------- */
+
+namespace DataflowVolumeToSpheres::Private
+{
+	static FVector ToVector(const openvdb::Vec3R& InPoint)
+	{
+		return FVector(InPoint.x(), InPoint.y(), InPoint.z());
+	}
+
+	static openvdb::Vec3R ToPoint(const FVector& InVector)
+	{
+		return openvdb::Vec3R(InVector.X, InVector.Y, InVector.Z);
+	}
+
+	static bool IsSphereContainedInSphere(const FVector& InCenterA,
+		const FVector& InCenterB,
+		const float InRadiusA,
+		const float InRadiusB)
+	{
+		if (InRadiusA > InRadiusB)
+		{
+			float Distance = (InCenterA - InCenterB).Length();
+
+			if (Distance <= (InRadiusA - InRadiusB))
+			{
+				return true;
+			}
+		}
+
+		return false;
+	}
+
+	static void ComputeDisplacement(const FVector& InCenterA,
+		const FVector& InCenterB,
+		const float InRadiusA,
+		const float InRadiusB,
+		const float InStepScalar,
+		FVector& OutDisplacementA,
+		FVector& OutDisplacementB)
+	{
+		OutDisplacementA.Set(0.f, 0.f, 0.f);
+		OutDisplacementB.Set(0.f, 0.f, 0.f);
+
+		FVector Diff = InCenterB - InCenterA;
+		float Distance = Diff.Length();
+		
+		if (Distance > UE_SMALL_NUMBER)
+		{
+			if (Distance <= (InRadiusA + InRadiusB) && InRadiusA > UE_SMALL_NUMBER && InRadiusB > UE_SMALL_NUMBER)
+			{
+				float Delta = InRadiusA + InRadiusB - Distance;
+
+				Diff.Normalize();
+
+				OutDisplacementA = -Diff * InStepScalar * Delta * InRadiusB / (InRadiusA + InRadiusB);
+				OutDisplacementB = Diff * InStepScalar * Delta * InRadiusA / (InRadiusA + InRadiusB);
+			}
+		}
+	}
+
+	static bool ComputeDisplacementAgainsSurface(const FVector& InCenter,
+		const float InRadius,
+		const float InDistanceThreshold,
+		const auto& InCsp,
+		const openvdb::FloatGrid::ConstAccessor& InAccessor,
+		const auto& InTransform,
+		FVector& OutDisplacement)
+	{
+		OutDisplacement.Set(0.f, 0.f, 0.f);
+
+		// Get closest point on the isosurface to InCenterA and InCenterB
+		std::vector<openvdb::Vec3R> Points;
+
+		Points.push_back(DataflowVolumeToSpheres::Private::ToPoint(InCenter));
+
+		std::vector<float> Distances;
+
+		if (!InCsp->searchAndReplace(Points, Distances))
+		{
+			return false;
+		}
+
+		FVector ClosestPoint = ToVector(Points[0]);
+
+		// Get phi values for InCenterA and InCenterB
+		openvdb::tools::GridSampler<openvdb::FloatGrid::ConstAccessor, openvdb::tools::BoxSampler>
+			fastSampler(InAccessor, InTransform);
+
+		float Phi = fastSampler.wsSample(DataflowVolumeToSpheres::Private::ToPoint(InCenter));
+
+		FVector Gradient = ClosestPoint - InCenter;
+
+		// If InCenter is on the isosurface then gradient is nonspecific and we can't displace
+		if (Gradient.Length() < UE_SMALL_NUMBER)
+		{
+			return true;
+		}
+
+		// Center is inside of isosurface
+		if (Phi < 0.f)
+		{
+			if (Gradient.Length() < InRadius && Gradient.Length() > UE_SMALL_NUMBER)
+			{
+				float Delta = InRadius - Gradient.Length();
+				if (Delta > InDistanceThreshold)
+				{
+					Gradient.Normalize();
+					OutDisplacement = -Gradient * (Delta - InDistanceThreshold);
+				}
+			}
+		}
+		else
+		{
+			float Delta = Gradient.Length() + InRadius - InDistanceThreshold;
+
+			Gradient.Normalize();
+			OutDisplacement = Gradient * Delta;
+		}
+
+		return true;
+	}
+}
+
+template<typename GridT, typename InterrupterT = openvdb::util::NullInterrupter>
+void VolumeToSpheresWithRelaxationProc(
+	const GridT& InGrid,
+	std::vector<openvdb::Vec4s>& OutSpheres,
+	int32 InPointCount,
+	const int32 InScatterRandomSeed,
+	const float InMinRadius,
+	const float InMaxRadius,
+	const int32 InRandomSeed,
+	float InIsovalue,
+	float InSpread,
+	bool InRemoveSpheres,
+	bool InApplyRelaxation,
+	int32 InSteps,
+	float InStepScalar,
+	bool InCorrectAgainstSurface,
+	float InDistanceThreshold,
+	InterrupterT* InInterrupter = nullptr)
+{
+	OutSpheres.clear();
+
+	if (InGrid.empty()) return;
+
+	auto gridPtr = InGrid.copy(); // shallow copy
+
+	if (gridPtr->getGridClass() == openvdb::GRID_LEVEL_SET) {
+		// Clamp the isovalue to the level set's background value minus epsilon.
+		// (In a valid narrow-band level set, all voxels, including background voxels,
+		// have values less than or equal to the background value, so an isovalue
+		// greater than or equal to the background value would produce a mask with
+		// effectively infinite extent.)
+		InIsovalue = std::min(InIsovalue,
+			static_cast<float>(gridPtr->background() - openvdb::math::Tolerance<float>::value()));
+	}
+	else if (gridPtr->getGridClass() == openvdb::GRID_FOG_VOLUME) {
+		// Clamp the isovalue of a fog volume between epsilon and one,
+		// again to avoid a mask with infinite extent.  (Recall that
+		// fog volume voxel values vary from zero outside to one inside.)
+		InIsovalue = openvdb::math::Clamp(InIsovalue, openvdb::math::Tolerance<float>::value(), 1.f);
+	}
+
+	// ClosestSurfacePoint is inaccurate for small grids.
+	// Resample the input grid if it is too small.
+	auto numVoxels = gridPtr->activeVoxelCount();
+
+	if (!numVoxels) return;
+
+	if (numVoxels < 10000) {
+		const auto scale = 1.0 / openvdb::math::Cbrt(2.0 * 10000.0 / double(numVoxels));
+		auto scaledXform = gridPtr->transform().copy();
+		scaledXform->preScale(scale);
+
+		auto newGridPtr = openvdb::tools::levelSetRebuild(*gridPtr, InIsovalue,
+			openvdb::LEVEL_SET_HALF_WIDTH, openvdb::LEVEL_SET_HALF_WIDTH, scaledXform.get(), InInterrupter);
+
+		const auto newNumVoxels = newGridPtr->activeVoxelCount();
+		if (newNumVoxels > numVoxels) {
+			OPENVDB_LOG_DEBUG_RUNTIME("fillWithSpheres: resampled input grid from "
+				<< numVoxels << " voxel" << (numVoxels == 1 ? "" : "s")
+				<< " to " << newNumVoxels << " voxel" << (newNumVoxels == 1 ? "" : "s"));
+			gridPtr = newGridPtr;
+			numVoxels = newNumVoxels;
+		}
+	}
+
+	const bool addNarrowBandPoints = (numVoxels < 10000);
+	int instances = InPointCount;
+
+	using TreeT = typename GridT::TreeType;
+	using BoolTreeT = typename TreeT::template ValueConverter<bool>::Type;
+	using Int16TreeT = typename TreeT::template ValueConverter<openvdb::Int16>::Type;
+
+	using RandGen = std::mersenne_twister_engine<uint32_t, 32, 351, 175, 19,
+		0xccab8ee7, 11, 0xffffffff, 7, 0x31b6ab00, 15, 0xffe50000, 17, 1812433253>; // mt11213b
+	RandGen mtRand(InScatterRandomSeed);
+
+	const TreeT& tree = gridPtr->tree();
+	openvdb::math::Transform transform = gridPtr->transform();
+
+	std::vector<openvdb::Vec3R> instancePoints;
+	{
+		// Compute a mask of the voxels enclosed by the isosurface.
+		typename openvdb::Grid<BoolTreeT>::Ptr interiorMaskPtr;
+		if (gridPtr->getGridClass() == openvdb::GRID_LEVEL_SET) {
+			interiorMaskPtr = openvdb::tools::sdfInteriorMask(*gridPtr, InIsovalue);
+		}
+		else {
+			// For non-level-set grids, the interior mask comprises the active voxels.
+			interiorMaskPtr = typename openvdb::Grid<BoolTreeT>::Ptr(openvdb::Grid<BoolTreeT>::create(false));
+			interiorMaskPtr->setTransform(transform.copy());
+			interiorMaskPtr->tree().topologyUnion(tree);
+		}
+
+		if (InInterrupter && InInterrupter->wasInterrupted()) return;
+
+		// If the interior mask is small and eroding it results in an empty grid,
+		// use the uneroded mask instead.  (But if the minimum sphere count is zero,
+		// then eroding away the mask is acceptable.)
+		if (!addNarrowBandPoints) {
+			openvdb::tools::erodeActiveValues(interiorMaskPtr->tree(), 1, openvdb::tools::NN_FACE, openvdb::tools::IGNORE_TILES);
+			openvdb::tools::pruneInactive(interiorMaskPtr->tree());
+		}
+		else {
+			auto& maskTree = interiorMaskPtr->tree();
+			auto copyOfTree = openvdb::StaticPtrCast<BoolTreeT>(maskTree.copy());
+			openvdb::tools::erodeActiveValues(maskTree, 1, openvdb::tools::NN_FACE, openvdb::tools::IGNORE_TILES);
+			openvdb::tools::pruneInactive(maskTree);
+			if (maskTree.empty()) { interiorMaskPtr->setTree(copyOfTree); }
+		}
+
+		// Scatter candidate sphere centroids (instancePoints)
+		instancePoints.reserve(instances);
+		openvdb::tools::v2s_internal::PointAccessor ptnAcc(instancePoints);
+
+		const auto scatterCount = openvdb::Index64(InPointCount);
+
+		openvdb::tools::UniformPointScatter<openvdb::tools::v2s_internal::PointAccessor, RandGen, InterrupterT> UniformScatter(
+				ptnAcc, scatterCount, mtRand, double(InSpread), InInterrupter);
+			UniformScatter(*interiorMaskPtr);
+
+		InPointCount = UniformScatter.getPointCount();
+	}
+
+	if (InInterrupter && InInterrupter->wasInterrupted()) return;
+
+	// Random radius
+	TArray<float> InstanceRadius;
+	InstanceRadius.SetNumUninitialized(InPointCount);
+
+	FRandomStream RandStream(InRandomSeed);
+
+	for (int32 Idx = 0; Idx < InPointCount; ++Idx)
+	{
+		InstanceRadius[Idx] = RandStream.FRandRange(InMinRadius, InMaxRadius);
+	}
+
+	// Remove spheres
+	TArray<bool> ValidSphere;
+	ValidSphere.Init(true, InPointCount);
+
+	if (InRemoveSpheres)
+	{
+		for (int32 Idx = 0; Idx < InPointCount; ++Idx)
+		{
+			if (!ValidSphere[Idx]) continue;
+
+			for (int32 IdxOther = Idx + 1; IdxOther < InPointCount; ++IdxOther)
+			{
+				if (!ValidSphere[IdxOther]) continue;
+
+				if (DataflowVolumeToSpheres::Private::IsSphereContainedInSphere(DataflowVolumeToSpheres::Private::ToVector(instancePoints[Idx]),
+					DataflowVolumeToSpheres::Private::ToVector(instancePoints[IdxOther]),
+					InstanceRadius[Idx],
+					InstanceRadius[IdxOther]))
+				{
+					ValidSphere[IdxOther] = false;
+				}
+			}
+		}
+	}
+
+	// Relaxation
+	if (InApplyRelaxation)
+	{
+		// For closestPoint query
+		auto csp = openvdb::tools::ClosestSurfacePoint<GridT>::create(*gridPtr, InIsovalue, InInterrupter);
+		if (!csp) return;
+
+		// For Phi value query
+		openvdb::FloatGrid::ConstAccessor Accessor = gridPtr->getConstAccessor();
+
+		for (int32 Idx_Step = 0; Idx_Step < InSteps; ++Idx_Step)
+		{
+			for (int32 Idx = 0; Idx < InPointCount; ++Idx)
+			{
+				if (!ValidSphere[Idx]) continue;
+
+				for (int32 IdxOther = Idx + 1; IdxOther < InPointCount; ++IdxOther)
+				{
+					if (!ValidSphere[IdxOther]) continue;
+
+					// Compute displacement
+					FVector DisplacementA, DisplacementB;
+
+					DataflowVolumeToSpheres::Private::ComputeDisplacement(DataflowVolumeToSpheres::Private::ToVector(instancePoints[Idx]),
+						DataflowVolumeToSpheres::Private::ToVector(instancePoints[IdxOther]),
+						InstanceRadius[Idx],
+						InstanceRadius[IdxOther],
+						InStepScalar,
+						DisplacementA,
+						DisplacementB);
+
+					// Add displacement
+					instancePoints[Idx].x() += DisplacementA.X;
+					instancePoints[Idx].y() += DisplacementA.Y;
+					instancePoints[Idx].z() += DisplacementA.Z;
+
+					instancePoints[IdxOther].x() += DisplacementB.X;
+					instancePoints[IdxOther].y() += DisplacementB.Y;
+					instancePoints[IdxOther].z() += DisplacementB.Z;
+
+					if (InCorrectAgainstSurface)
+					{
+						// Check against surface
+						FVector Displacement;
+
+						if (DataflowVolumeToSpheres::Private::ComputeDisplacementAgainsSurface(DataflowVolumeToSpheres::Private::ToVector(instancePoints[Idx]),
+							InstanceRadius[Idx],
+							InDistanceThreshold,
+							csp,
+							Accessor,
+							gridPtr->transform(),
+							Displacement))
+						{
+							// Add displacement
+							instancePoints[Idx].x() += Displacement.X;
+							instancePoints[Idx].y() += Displacement.Y;
+							instancePoints[Idx].z() += Displacement.Z;
+						}
+
+						if (DataflowVolumeToSpheres::Private::ComputeDisplacementAgainsSurface(DataflowVolumeToSpheres::Private::ToVector(instancePoints[IdxOther]),
+							InstanceRadius[IdxOther],
+							InDistanceThreshold,
+							csp,
+							Accessor,
+							gridPtr->transform(),
+							Displacement))
+
+						{
+							// Add displacement
+							instancePoints[IdxOther].x() += Displacement.X;
+							instancePoints[IdxOther].y() += Displacement.Y;
+							instancePoints[IdxOther].z() += Displacement.Z;
+						}
+					}
+				}
+			}
+		}
+	}
+
+	// Create spheres
+	int32 NumValidSpheres = 0;
+
+	for (int32 Idx = 0; Idx < InPointCount; ++Idx)
+	{
+		if (ValidSphere[Idx])
+		{
+			NumValidSpheres++;
+		}
+	}
+
+	OutSpheres.reserve(NumValidSpheres);
+
+	for (int32 Idx = 0; Idx < InPointCount; ++Idx)
+	{
+		if (ValidSphere[Idx])
+		{
+			const openvdb::Vec4s sphere(
+				float(instancePoints[Idx].x()),
+				float(instancePoints[Idx].y()),
+				float(instancePoints[Idx].z()),
+				InstanceRadius[Idx]);
+
+			OutSpheres.push_back(sphere);
+		}
+	}
+}
 
diff --git a/Engine/Plugins/Experimental/DataflowVolumetricPlugin/Source/DataflowVolumeNodes/Private/Dataflow/DataflowVolumeNodes.cpp b/Engine/Plugins/Experimental/DataflowVolumetricPlugin/Source/DataflowVolumeNodes/Private/Dataflow/DataflowVolumeNodes.cpp
--- a/Engine/Plugins/Experimental/DataflowVolumetricPlugin/Source/DataflowVolumeNodes/Private/Dataflow/DataflowVolumeNodes.cpp
+++ b/Engine/Plugins/Experimental/DataflowVolumetricPlugin/Source/DataflowVolumeNodes/Private/Dataflow/DataflowVolumeNodes.cpp
@@ -47,6 +47,8 @@ namespace UE::Dataflow
 
 		// Terminal
 		DATAFLOW_NODE_REGISTER_CREATION_FACTORY(FVolumeTextureTerminalNode);
+
+		DATAFLOW_NODE_REGISTER_CREATION_FACTORY(FVolumeToSpheresWithRelaxationDataflowNode);
 	}
 }
 
@@ -683,11 +685,6 @@ FVolumeToSpheresDataflowNode_v2::FVolumeToSpheresDataflowNode_v2(const UE::Dataf
 	RegisterInputConnection(&MaxNumberOfPointsPerVoxel).SetCanHidePin(true).SetPinIsHidden(true);
 	RegisterOutputConnection(&Collection);
 	RegisterOutputConnection(&Spheres);
-
-	DebugDrawRenderSettings.RenderType = EDataflowDebugDrawRenderType::Wireframe;
-	DebugDrawRenderSettings.bTranslucent = false;
-	DebugDrawRenderSettings.Color = FLinearColor::Blue;
-	DebugDrawRenderSettings.LineWidthMultiplier = 1.5f;
 }
 
 void FVolumeToSpheresDataflowNode_v2::Evaluate(UE::Dataflow::FContext& Context, const FDataflowOutput* Out) const
@@ -787,82 +784,6 @@ void FVolumeToSpheresDataflowNode_v2::Evaluate(UE::Dataflow::FContext& Context,
 	}
 }
 
-#if WITH_EDITOR
-bool FVolumeToSpheresDataflowNode_v2::CanDebugDrawViewMode(const FName& ViewModeName) const
-{
-	return ViewModeName == UE::Dataflow::FDataflowConstruction3DViewMode::Name;
-}
-
-void FVolumeToSpheresDataflowNode_v2::DebugDraw(UE::Dataflow::FContext& Context, IDataflowDebugDrawInterface& DataflowRenderingInterface, const FDebugDrawParameters& DebugDrawParameters) const
-{
-	using namespace UE::Dataflow::Overlay;
-
-	if (DebugDrawParameters.bNodeIsSelected || DebugDrawParameters.bNodeIsPinned)
-	{
-		FString OverlayEvaluateResultString;
-
-		const FDataflowOutput* SpheresOutput = FindOutput(&Spheres);
-		if (SpheresOutput)
-		{
-			const TArray<FSphere>& OutSpheres = SpheresOutput->GetValue(Context, Spheres);
-			//			const FDataflowFloatVolume& OutFloatVolume = this->GetOutputValue(Context, FloatVolume);
-
-			DebugDrawRenderSettings.SetDebugDrawSettings(DataflowRenderingInterface);
-
-			float MinRad = FLT_MAX, MaxRad = -FLT_MAX, SumRad = 0;
-			for (const FSphere& Sphere : OutSpheres)
-			{
-				if (Sphere.W < MinRad)
-				{
-					MinRad = Sphere.W;
-				}
-				if (Sphere.W > MaxRad)
-				{
-					MaxRad = Sphere.W;
-				}
-
-				SumRad += Sphere.W;
-
-				if (bDrawSphereCenters)
-				{
-					DataflowRenderingInterface.DrawPoint(Sphere.Center);
-				}
-
-				if (bDrawSpheres)
-				{
-					DataflowRenderingInterface.DrawSphere(Sphere.Center, Sphere.W);
-				}
-			}
-
-			if (OutSpheres.Num())
-			{
-				const int32 NumSpheres = OutSpheres.Num();
-				const float MeanRad = SumRad / (float)NumSpheres;
-
-				OverlayEvaluateResultString.Appendf(TEXT("Number of Spheres: %d\nMin radius: %3.3f\nMax radius: %3.3f\nMedian radius: %3.3f\n"),
-					NumSpheres,
-					MinRad,
-					MaxRad,
-					MeanRad);
-			}
-			else
-			{
-				OverlayEvaluateResultString.Appendf(TEXT("Number of Spheres: 0\nMin radius: 0.000\nMax radius: 0.000\nMedian radius: 0.000\n"));
-			}
-		}
-
-		// ------------------------------------------------------------------------------------------------------
-		// Overlay
-		TArray<FString> OverlayStrArr;
-		OverlayStrArr.Add(BuildOverlayNodeInfoString(this));
-		OverlayStrArr.Add(OverlayEvaluateResultString);
-		FString OverlayStr = BuildOverlayFinalString(OverlayStrArr);
-		DataflowRenderingInterface.DrawOverlayText(OverlayStr);
-		// ------------------------------------------------------------------------------------------------------
-	}
-}
-#endif
-
 /* --------------------------------------------------------------------------------------------- */
 
 FConvertVolumeDataflowNode::FConvertVolumeDataflowNode(const UE::Dataflow::FNodeParameters& InParam, FGuid InGuid)
@@ -2172,3 +2093,61 @@ void FVolumeTextureTerminalNode::SetAssetValue(TObjectPtr<UObject> Asset, UE::Da
 
 /* --------------------------------------------------------------------------------------------- */
 
+FVolumeToSpheresWithRelaxationDataflowNode::FVolumeToSpheresWithRelaxationDataflowNode(const UE::Dataflow::FNodeParameters& InParam, FGuid InGuid)
+	: FDataflowNode(InParam, InGuid)
+{
+	RegisterInputConnection(&FloatVolume);
+	RegisterInputConnection(&PointCount).SetCanHidePin(true).SetPinIsHidden(true);
+	RegisterInputConnection(&ScatterRandomSeed).SetCanHidePin(true).SetPinIsHidden(true);
+	RegisterInputConnection(&Spread).SetCanHidePin(true).SetPinIsHidden(true);
+	RegisterInputConnection(&MinRadius).SetCanHidePin(true).SetPinIsHidden(true);
+	RegisterInputConnection(&MaxRadius).SetCanHidePin(true).SetPinIsHidden(true);
+	RegisterInputConnection(&RandomSeed).SetCanHidePin(true).SetPinIsHidden(true);
+	RegisterInputConnection(&IsoValue).SetCanHidePin(true).SetPinIsHidden(true);
+	RegisterInputConnection(&Steps).SetCanHidePin(true).SetPinIsHidden(true);
+	RegisterInputConnection(&StepScalar).SetCanHidePin(true).SetPinIsHidden(true);
+	RegisterInputConnection(&DistanceThreshold).SetCanHidePin(true).SetPinIsHidden(true);
+	RegisterOutputConnection(&Spheres);
+}
+
+void FVolumeToSpheresWithRelaxationDataflowNode::Evaluate(UE::Dataflow::FContext& Context, const FDataflowOutput* Out) const
+{
+	if (Out->IsA(&Spheres))
+	{
+		TArray<FSphere> OutSpheres;
+
+		if (IsConnected(&FloatVolume))
+		{
+			const FDataflowFloatVolume& InFloatVolume = GetValue(Context, &FloatVolume);
+			const int32 InPointCount = GetValue(Context, &PointCount);
+			const int32 InScatterRandomSeed = GetValue(Context, &ScatterRandomSeed);
+			const float InSpread = GetValue(Context, &Spread);
+			const float InMinRadius = GetValue(Context, &MinRadius);
+			const float InMaxRadius = GetValue(Context, &MaxRadius);
+			const int32 InRandomSeed = GetValue(Context, &RandomSeed);
+			const float InIsoValue = GetValue(Context, &IsoValue);
+			const int32 InSteps = GetValue(Context, &Steps);
+			const float InStepScalar = GetValue(Context, &StepScalar);
+			const float InDistanceThreshold = GetValue(Context, &DistanceThreshold);
+
+			// Compute sphere packing
+			InFloatVolume.VolumeToSpheresWithRelaxation(
+				InPointCount,
+				InScatterRandomSeed,
+				InMinRadius,
+				InMaxRadius,
+				InRandomSeed,
+				InIsoValue,
+				InSpread,
+				bRemoveEmbededSpheres,
+				bApplyRelaxation,
+				InSteps,
+				InStepScalar,
+				bCorrectAgainstSurface,
+				InDistanceThreshold,
+				OutSpheres);
+		}
+
+		SetValue(Context, MoveTemp(OutSpheres), &Spheres);
+	}
+}
diff --git a/Engine/Plugins/Experimental/DataflowVolumetricPlugin/Source/DataflowVolumeNodes/Private/Dataflow/DataflowVolumeNodes.h b/Engine/Plugins/Experimental/DataflowVolumetricPlugin/Source/DataflowVolumeNodes/Private/Dataflow/DataflowVolumeNodes.h
--- a/Engine/Plugins/Experimental/DataflowVolumetricPlugin/Source/DataflowVolumeNodes/Private/Dataflow/DataflowVolumeNodes.h
+++ b/Engine/Plugins/Experimental/DataflowVolumetricPlugin/Source/DataflowVolumeNodes/Private/Dataflow/DataflowVolumeNodes.h
@@ -437,15 +437,6 @@ private:
 
 	virtual void Evaluate(UE::Dataflow::FContext& Context, const FDataflowOutput* Out) const override;
 
-#if WITH_EDITOR
-	virtual bool CanDebugDraw() const override
-	{
-		return true;
-	}
-	virtual bool CanDebugDrawViewMode(const FName& ViewModeName) const override;
-	virtual void DebugDraw(UE::Dataflow::FContext& Context, IDataflowDebugDrawInterface& DataflowRenderingInterface, const FDebugDrawParameters& DebugDrawParameters) const override;
-#endif
-
 public:
 	FVolumeToSpheresDataflowNode_v2(const UE::Dataflow::FNodeParameters& InParam, FGuid InGuid = FGuid::NewGuid());
 };
@@ -967,6 +958,105 @@ private:
 public:
 	FVolumeTextureTerminalNode(const UE::Dataflow::FNodeParameters& InParam, FGuid InGuid = FGuid::NewGuid());
 };
+
+/**
+ *
+ * 
+ *
+ */
+USTRUCT(meta = (Experimental, DataflowVolume))
+struct FVolumeToSpheresWithRelaxationDataflowNode : public FDataflowNode
+{
+	GENERATED_USTRUCT_BODY()
+	DATAFLOW_NODE_DEFINE_INTERNAL(FVolumeToSpheresWithRelaxationDataflowNode, "VolumeToSpheresWithRelaxation", "Volume|Utilities", "")
+
+private:
+	/** Float volume */
+	UPROPERTY(meta = (DataflowInput))
+	FDataflowFloatVolume FloatVolume;
+
+	/** Number of points to scatter inside of the volume */
+	UPROPERTY(EditAnywhere, Category = "PointScatter", meta = (DataflowInput, UIMin = 0, ClampMin = 0));
+	int32 PointCount = 100;
+
+	/** Random seed for scatter */
+	UPROPERTY(EditAnywhere, Category = "PointScatter", meta = (DataflowInput, UIMin = 0, ClampMin = 0));
+	int32 ScatterRandomSeed = 0;
+
+	/** Spread for OpenVDB scatter */
+	UPROPERTY(EditAnywhere, Category = "PointScatter", meta = (DataflowInput, UIMin = 0, ClampMin = 0, UIMax = 1, ClampMax = 1));
+	float Spread = 1.f;
+
+	/** Minimum radius for spheres */
+	UPROPERTY(EditAnywhere, Category = "Spheres", meta = (DataflowInput, DisplayName = "Minimum Radius", UIMin = 0.001, ClampMin = 0.001f));
+	float MinRadius = 1.f;
+
+	/** Maximum radius for spheres */
+	UPROPERTY(EditAnywhere, Category = "Spheres", meta = (DataflowInput, DisplayName = "Maximum Radius", UIMin = 0.001, ClampMin = 0.001f));
+	float MaxRadius = 10.f;
+
+	/** Random seed for radius */
+	UPROPERTY(EditAnywhere, Category = "Spheres", meta = (DataflowInput, UIMin = 0, ClampMin = 0));
+	int32 RandomSeed = 0;
+
+	/** Isovalue used for scattering */
+	UPROPERTY(EditAnywhere, Category = "Spheres", meta = (DataflowInput));
+	float IsoValue = 0.f;
+
+	/** Remove spheres which are completely inside of another sphere */
+	UPROPERTY(EditAnywhere, Category = "Spheres");
+	bool bRemoveEmbededSpheres = true;
+
+	/** Apply relaxation to the remaining spheres after removal */
+	UPROPERTY(EditAnywhere, Category = "Relaxation");
+	bool bApplyRelaxation = true;
+
+	/** Number of steps in the relaxation */
+	UPROPERTY(EditAnywhere, Category = "Relaxation", meta = (DataflowInput, UIMin = 1, ClampMin = 1, EditCondition = "bApplyRelaxation"));
+	int32 Steps = 10;
+
+	/** Scale the displacement applied to separate spheres */
+	UPROPERTY(EditAnywhere, Category = "Relaxation", meta = (DataflowInput, UIMin = 0, ClampMin = 0.f, EditCondition = "bApplyRelaxation"));
+	float StepScalar = 1.f;
+
+	/** Move spheres so their distance from the isosurface equals to DistanceThreshold */
+	UPROPERTY(EditAnywhere, Category = "Relaxation", meta = (DisplayName = "Correct against Isosurface", EditCondition = "bApplyRelaxation"));
+	bool bCorrectAgainstSurface = true;
+	
+	/** Amount allowed of a sphere to be outside of the isosurface */
+	UPROPERTY(EditAnywhere, Category = "Relaxation", meta = (DataflowInput, ClampMin = 0.f, EditCondition = "bApplyRelaxation && bCorrectAgainstSurface"));
+	float DistanceThreshold = 0.f;
+
+	/** Output spheres */
+	UPROPERTY(meta = (DataflowOutput))
+	TArray<FSphere> Spheres;
+
+	virtual void Evaluate(UE::Dataflow::FContext& Context, const FDataflowOutput* Out) const override;
+
+public:
+	FVolumeToSpheresWithRelaxationDataflowNode(const UE::Dataflow::FNodeParameters& InParam, FGuid InGuid = FGuid::NewGuid());
+};
+
+
+
+
+
+
+
+
+
+
+
+
+
+
+
+
+
+
+
+
+
 
 namespace UE::Dataflow
 {
