Dataflow:
- Removed DebugDraw from FDataflowVolumeToSpheres v1 and v2 (all the debugdraw settings now hidden)
- Removed DebugDraw from all the volume nodes
- Added new rendering to DataflowVolume
- bug fixes

#Synced-CL 50146089

diff --git a/Engine/Plugins/Experimental/DataflowVolumetricPlugin/Source/DataflowVolumeCore/Internal/Dataflow/DataflowVolume.h b/Engine/Plugins/Experimental/DataflowVolumetricPlugin/Source/DataflowVolumeCore/Internal/Dataflow/DataflowVolume.h
--- a/Engine/Plugins/Experimental/DataflowVolumetricPlugin/Source/DataflowVolumeCore/Internal/Dataflow/DataflowVolume.h
+++ b/Engine/Plugins/Experimental/DataflowVolumetricPlugin/Source/DataflowVolumeCore/Internal/Dataflow/DataflowVolume.h
@@ -23,8 +23,8 @@ public:
 	FDataflowVolume() {}
 	virtual ~FDataflowVolume() {}
 
-	DATAFLOWVOLUMECORE_API int32 GetNumActiveVoxels() const;
-	DATAFLOWVOLUMECORE_API void GetActiveVoxels(TArray<FBox>& OutActiveVoxels) const;
+	DATAFLOWVOLUMECORE_API int32 GetNumActiveVoxels(const float InIsovalue, bool bInteriorMaskOnly = false) const;
+	DATAFLOWVOLUMECORE_API void GetActiveVoxels(const float InIsovalue, TArray<FBox>& OutActiveVoxels, bool bInteriorMaskOnly = false) const;
 	DATAFLOWVOLUMECORE_API FBox GetVolumeBoundingBox() const;
 
 	DATAFLOWVOLUMECORE_API bool CreateVolumeTexture(UVolumeTexture* InVolumeTexture) const;
diff --git a/Engine/Plugins/Experimental/DataflowVolumetricPlugin/Source/DataflowVolumeCore/Private/Dataflow/DataflowBoolVolumeImpl.cpp b/Engine/Plugins/Experimental/DataflowVolumetricPlugin/Source/DataflowVolumeCore/Private/Dataflow/DataflowBoolVolumeImpl.cpp
--- a/Engine/Plugins/Experimental/DataflowVolumetricPlugin/Source/DataflowVolumeCore/Private/Dataflow/DataflowBoolVolumeImpl.cpp
+++ b/Engine/Plugins/Experimental/DataflowVolumetricPlugin/Source/DataflowVolumeCore/Private/Dataflow/DataflowBoolVolumeImpl.cpp
@@ -3,6 +3,7 @@
 #include "Dataflow/DataflowBoolVolumeImpl.h"
 #include <Chaos/ArrayND.h>
 #include "Engine/VolumeTexture.h"
+#include "Dataflow/DataflowVolumeUtils.h"
 
 namespace UE::DataflowVolume::Private
 {
@@ -58,35 +59,11 @@ namespace UE::DataflowVolume::Private
 
 	/* --------------------------------------------------------------------------------------------------------------- */
 
-	void FDataflowBoolVolumeImpl::GetActiveVoxels(TArray<FBox>& OutActiveVoxels) const
+	void FDataflowBoolVolumeImpl::GetActiveVoxels(const float InIsovalue, TArray<FBox>& OutActiveVoxels, bool bInteriorMaskOnly) const
 	{
 		if (BoolGrid)
 		{
-			const openvdb::Index64 ActiveVoxelCount = BoolGrid->activeVoxelCount();
-			OutActiveVoxels.AddUninitialized(ActiveVoxelCount);
-
-			openvdb::math::Transform T = this->BoolGrid->transform();
-
-			int32 Idx = 0;
-			for (openvdb::BoolGrid::ValueOnCIter iter = this->BoolGrid->cbeginValueOn(); iter.test(); ++iter)
-			{
-				if (iter.isVoxelValue())
-				{
-					openvdb::CoordBBox BBox;
-					iter.getBoundingBox(BBox);
-
-					const openvdb::Vec3d BBoxMin(BBox.min().x() - 0.5, BBox.min().y() - 0.5, BBox.min().z() - 0.5);
-					const openvdb::Vec3d BBoxMax(BBox.max().x() + 0.5, BBox.max().y() + 0.5, BBox.max().z() + 0.5);
-
-					const openvdb::Vec3d BBoxMinWS = BoolGrid->transform().indexToWorld(BBoxMin);
-					const openvdb::Vec3d BBoxMaxWS = BoolGrid->transform().indexToWorld(BBoxMax);
-
-					FVector Min(BBoxMinWS.x(), BBoxMinWS.y(), BBoxMinWS.z());
-					FVector Max(BBoxMaxWS.x(), BBoxMaxWS.y(), BBoxMaxWS.z());
-
-					OutActiveVoxels[Idx++] = FBox(Min, Max);
-				}
-			}
+			UE::DataflowVolumeUtils::GetActiveVoxels<openvdb::BoolGrid>(BoolGrid, InIsovalue, OutActiveVoxels, bInteriorMaskOnly);
 		}
 	}
 
@@ -108,7 +85,7 @@ namespace UE::DataflowVolume::Private
 
 		const EPixelFormat PixelFormat = EPixelFormat::PF_R16F;
 
-		if (!BoolGrid->empty())
+		if (BoolGrid && !BoolGrid->empty())
 		{
 			openvdb::Coord ActiveVoxelDim = BoolGrid->evalActiveVoxelDim();
 			openvdb::CoordBBox Bbox = BoolGrid->evalActiveVoxelBoundingBox();
diff --git a/Engine/Plugins/Experimental/DataflowVolumetricPlugin/Source/DataflowVolumeCore/Private/Dataflow/DataflowBoolVolumeImpl.h b/Engine/Plugins/Experimental/DataflowVolumetricPlugin/Source/DataflowVolumeCore/Private/Dataflow/DataflowBoolVolumeImpl.h
--- a/Engine/Plugins/Experimental/DataflowVolumetricPlugin/Source/DataflowVolumeCore/Private/Dataflow/DataflowBoolVolumeImpl.h
+++ b/Engine/Plugins/Experimental/DataflowVolumetricPlugin/Source/DataflowVolumeCore/Private/Dataflow/DataflowBoolVolumeImpl.h
@@ -17,7 +17,7 @@ namespace UE::DataflowVolume::Private
 
 	public:
 		// Virtual functions
-		virtual void GetActiveVoxels(TArray<FBox>& OutActiveVoxels) const override;
+		virtual void GetActiveVoxels(const float InIsovalue, TArray<FBox>& OutActiveVoxels, bool bInteriorMaskOnly = false) const override;
 		virtual FBox GetVolumeBoundingBox() const override;
 		virtual bool CreateVolumeTexture(UVolumeTexture* InVolumeTexture) const override;
 
diff --git a/Engine/Plugins/Experimental/DataflowVolumetricPlugin/Source/DataflowVolumeCore/Private/Dataflow/DataflowFloatVectorVolumeImpl.cpp b/Engine/Plugins/Experimental/DataflowVolumetricPlugin/Source/DataflowVolumeCore/Private/Dataflow/DataflowFloatVectorVolumeImpl.cpp
--- a/Engine/Plugins/Experimental/DataflowVolumetricPlugin/Source/DataflowVolumeCore/Private/Dataflow/DataflowFloatVectorVolumeImpl.cpp
+++ b/Engine/Plugins/Experimental/DataflowVolumetricPlugin/Source/DataflowVolumeCore/Private/Dataflow/DataflowFloatVectorVolumeImpl.cpp
@@ -4,6 +4,7 @@
 #include "Dataflow/DataflowBoolVolumeImpl.h"
 #include <Dataflow/DataflowBoolVolume.h>
 #include "Engine/VolumeTexture.h"
+#include "Dataflow/DataflowVolumeUtils.h"
 
 namespace UE::DataflowVolume::Private
 {
@@ -144,32 +145,11 @@ namespace UE::DataflowVolume::Private
 
 	/* --------------------------------------------------------------------------------------------------------------- */
 
-	void FDataflowFloatVectorVolumeImpl::GetActiveVoxels(TArray<FBox>& OutActiveVoxels) const
+	void FDataflowFloatVectorVolumeImpl::GetActiveVoxels(const float InIsovalue, TArray<FBox>& OutActiveVoxels, bool bInteriorMaskOnly) const
 	{
-		const openvdb::Index64 ActiveVoxelCount = FloatVectorGrid->activeVoxelCount();
-		OutActiveVoxels.AddUninitialized(ActiveVoxelCount);
-
-		openvdb::math::Transform T = this->FloatVectorGrid->transform();
-
-		int32 Idx = 0;
-		for (openvdb::VectorGrid::ValueOnCIter iter = this->FloatVectorGrid->cbeginValueOn(); iter.test(); ++iter)
+		if (FloatVectorGrid)
 		{
-			if (iter.isVoxelValue())
-			{
-				openvdb::CoordBBox BBox;
-				iter.getBoundingBox(BBox);
-
-				const openvdb::Vec3d BBoxMin(BBox.min().x() - 0.5, BBox.min().y() - 0.5, BBox.min().z() - 0.5);
-				const openvdb::Vec3d BBoxMax(BBox.max().x() + 0.5, BBox.max().y() + 0.5, BBox.max().z() + 0.5);
-
-				const openvdb::Vec3d BBoxMinWS = FloatVectorGrid->transform().indexToWorld(BBoxMin);
-				const openvdb::Vec3d BBoxMaxWS = FloatVectorGrid->transform().indexToWorld(BBoxMax);
-
-				FVector Min(BBoxMinWS.x(), BBoxMinWS.y(), BBoxMinWS.z());
-				FVector Max(BBoxMaxWS.x(), BBoxMaxWS.y(), BBoxMaxWS.z());
-
-				OutActiveVoxels[Idx++] = FBox(Min, Max);
-			}
+			UE::DataflowVolumeUtils::GetActiveVoxels<openvdb::VectorGrid>(FloatVectorGrid, InIsovalue, OutActiveVoxels, bInteriorMaskOnly);
 		}
 	}
 
diff --git a/Engine/Plugins/Experimental/DataflowVolumetricPlugin/Source/DataflowVolumeCore/Private/Dataflow/DataflowFloatVectorVolumeImpl.h b/Engine/Plugins/Experimental/DataflowVolumetricPlugin/Source/DataflowVolumeCore/Private/Dataflow/DataflowFloatVectorVolumeImpl.h
--- a/Engine/Plugins/Experimental/DataflowVolumetricPlugin/Source/DataflowVolumeCore/Private/Dataflow/DataflowFloatVectorVolumeImpl.h
+++ b/Engine/Plugins/Experimental/DataflowVolumetricPlugin/Source/DataflowVolumeCore/Private/Dataflow/DataflowFloatVectorVolumeImpl.h
@@ -23,7 +23,7 @@ namespace UE::DataflowVolume::Private
 
 	public:
 		// Virtual functions
-		virtual void GetActiveVoxels(TArray<FBox>& OutActiveVoxels) const override;
+		DATAFLOWVOLUMECORE_API void GetActiveVoxels(const float InIsovalue, TArray<FBox>& OutActiveVoxels, bool bInteriorMaskOnly = false) const;
 		virtual FBox GetVolumeBoundingBox() const override;
 		virtual bool CreateVolumeTexture(UVolumeTexture* InVolumeTexture) const override;
 
diff --git a/Engine/Plugins/Experimental/DataflowVolumetricPlugin/Source/DataflowVolumeCore/Private/Dataflow/DataflowFloatVolumeImpl.cpp b/Engine/Plugins/Experimental/DataflowVolumetricPlugin/Source/DataflowVolumeCore/Private/Dataflow/DataflowFloatVolumeImpl.cpp
--- a/Engine/Plugins/Experimental/DataflowVolumetricPlugin/Source/DataflowVolumeCore/Private/Dataflow/DataflowFloatVolumeImpl.cpp
+++ b/Engine/Plugins/Experimental/DataflowVolumetricPlugin/Source/DataflowVolumeCore/Private/Dataflow/DataflowFloatVolumeImpl.cpp
@@ -8,6 +8,7 @@
 #include "Dataflow/DataflowVolumeUtils.h"
 #include "Dataflow/DataflowVolumeAlgo.h"
 #include "Engine/VolumeTexture.h"
+#include "Dataflow/DataflowVolumeUtils.h"
 
 namespace UE::DataflowVolume::Private
 {
@@ -513,33 +514,13 @@ namespace UE::DataflowVolume::Private
 
 	/* --------------------------------------------------------------------------------------------------------------- */
 
-	void FDataflowFloatVolumeImpl::GetActiveVoxels(TArray<FBox>& OutActiveVoxels) const
-	{
-		const openvdb::Index64 ActiveVoxelCount = FloatGrid->activeVoxelCount();
-		OutActiveVoxels.AddUninitialized(ActiveVoxelCount);
-
-		openvdb::math::Transform T = this->FloatGrid->transform();
 
-		int32 Idx = 0;
-		for (openvdb::FloatGrid::ValueOnCIter iter = this->FloatGrid->cbeginValueOn(); iter.test(); ++iter)
+	void FDataflowFloatVolumeImpl::GetActiveVoxels(const float InIsovalue, TArray<FBox>& OutActiveVoxels, bool bInteriorMaskOnly) const
+	{
+		if (FloatGrid)
 		{
-			if (iter.isVoxelValue())
-			{
-				openvdb::CoordBBox BBox;
-				iter.getBoundingBox(BBox);
-
-				const openvdb::Vec3d BBoxMin(BBox.min().x() - 0.5, BBox.min().y() - 0.5, BBox.min().z() - 0.5);
-				const openvdb::Vec3d BBoxMax(BBox.max().x() + 0.5, BBox.max().y() + 0.5, BBox.max().z() + 0.5);
-
-				const openvdb::Vec3d BBoxMinWS = FloatGrid->transform().indexToWorld(BBoxMin);
-				const openvdb::Vec3d BBoxMaxWS = FloatGrid->transform().indexToWorld(BBoxMax);
-
-				FVector Min(BBoxMinWS.x(), BBoxMinWS.y(), BBoxMinWS.z());
-				FVector Max(BBoxMaxWS.x(), BBoxMaxWS.y(), BBoxMaxWS.z());
-
-				OutActiveVoxels[Idx++] = FBox(Min, Max);
-			}
-		}
+			UE::DataflowVolumeUtils::GetActiveVoxels<openvdb::FloatGrid>(FloatGrid, InIsovalue, OutActiveVoxels, bInteriorMaskOnly);
+		}	
 	}
 
 	/* --------------------------------------------------------------------------------------------------------------- */
diff --git a/Engine/Plugins/Experimental/DataflowVolumetricPlugin/Source/DataflowVolumeCore/Private/Dataflow/DataflowFloatVolumeImpl.h b/Engine/Plugins/Experimental/DataflowVolumetricPlugin/Source/DataflowVolumeCore/Private/Dataflow/DataflowFloatVolumeImpl.h
--- a/Engine/Plugins/Experimental/DataflowVolumetricPlugin/Source/DataflowVolumeCore/Private/Dataflow/DataflowFloatVolumeImpl.h
+++ b/Engine/Plugins/Experimental/DataflowVolumetricPlugin/Source/DataflowVolumeCore/Private/Dataflow/DataflowFloatVolumeImpl.h
@@ -30,7 +30,7 @@ namespace UE::DataflowVolume::Private
 	public:
 		// Virtual functions
 		virtual FString VolumeInfo() const override;
-		virtual void GetActiveVoxels(TArray<FBox>& OutActiveVoxels) const override;
+		virtual void GetActiveVoxels(const float InIsovalue, TArray<FBox>& OutActiveVoxels, bool bInteriorMaskOnly = false) const override;
 		virtual FBox GetVolumeBoundingBox() const override;
 		virtual bool CreateVolumeTexture(UVolumeTexture* InVolumeTexture) const override;
 
diff --git a/Engine/Plugins/Experimental/DataflowVolumetricPlugin/Source/DataflowVolumeCore/Private/Dataflow/DataflowIntVolumeImpl.cpp b/Engine/Plugins/Experimental/DataflowVolumetricPlugin/Source/DataflowVolumeCore/Private/Dataflow/DataflowIntVolumeImpl.cpp
--- a/Engine/Plugins/Experimental/DataflowVolumetricPlugin/Source/DataflowVolumeCore/Private/Dataflow/DataflowIntVolumeImpl.cpp
+++ b/Engine/Plugins/Experimental/DataflowVolumetricPlugin/Source/DataflowVolumeCore/Private/Dataflow/DataflowIntVolumeImpl.cpp
@@ -3,6 +3,7 @@
 #include "Dataflow/DataflowIntVolumeImpl.h"
 #include <Chaos/ArrayND.h>
 #include "Engine/VolumeTexture.h"
+#include "Dataflow/DataflowVolumeUtils.h"
 
 namespace UE::DataflowVolume::Private
 {
@@ -84,32 +85,11 @@ namespace UE::DataflowVolume::Private
 
 	/* --------------------------------------------------------------------------------------------------------------- */
 
-	void FDataflowIntVolumeImpl::GetActiveVoxels(TArray<FBox>& OutActiveVoxels) const
+	void FDataflowIntVolumeImpl::GetActiveVoxels(const float InIsovalue, TArray<FBox>& OutActiveVoxels, bool bInteriorMaskOnly) const
 	{
-		const openvdb::Index64 ActiveVoxelCount = IntGrid->activeVoxelCount();
-		OutActiveVoxels.AddUninitialized(ActiveVoxelCount);
-
-		openvdb::math::Transform T = this->IntGrid->transform();
-
-		int32 Idx = 0;
-		for (openvdb::Int32Grid::ValueOnCIter iter = this->IntGrid->cbeginValueOn(); iter.test(); ++iter)
+		if (IntGrid)
 		{
-			if (iter.isVoxelValue())
-			{
-				openvdb::CoordBBox BBox;
-				iter.getBoundingBox(BBox);
-
-				const openvdb::Vec3d BBoxMin(BBox.min().x() - 0.5, BBox.min().y() - 0.5, BBox.min().z() - 0.5);
-				const openvdb::Vec3d BBoxMax(BBox.max().x() + 0.5, BBox.max().y() + 0.5, BBox.max().z() + 0.5);
-
-				const openvdb::Vec3d BBoxMinWS = IntGrid->transform().indexToWorld(BBoxMin);
-				const openvdb::Vec3d BBoxMaxWS = IntGrid->transform().indexToWorld(BBoxMax);
-
-				FVector Min(BBoxMinWS.x(), BBoxMinWS.y(), BBoxMinWS.z());
-				FVector Max(BBoxMaxWS.x(), BBoxMaxWS.y(), BBoxMaxWS.z());
-
-				OutActiveVoxels[Idx++] = FBox(Min, Max);
-			}
+			UE::DataflowVolumeUtils::GetActiveVoxels<openvdb::Int32Grid>(IntGrid, InIsovalue, OutActiveVoxels, bInteriorMaskOnly);
 		}
 	}
 
diff --git a/Engine/Plugins/Experimental/DataflowVolumetricPlugin/Source/DataflowVolumeCore/Private/Dataflow/DataflowIntVolumeImpl.h b/Engine/Plugins/Experimental/DataflowVolumetricPlugin/Source/DataflowVolumeCore/Private/Dataflow/DataflowIntVolumeImpl.h
--- a/Engine/Plugins/Experimental/DataflowVolumetricPlugin/Source/DataflowVolumeCore/Private/Dataflow/DataflowIntVolumeImpl.h
+++ b/Engine/Plugins/Experimental/DataflowVolumetricPlugin/Source/DataflowVolumeCore/Private/Dataflow/DataflowIntVolumeImpl.h
@@ -19,7 +19,7 @@ namespace UE::DataflowVolume::Private
 	public:
 		// Virtual functions
 		virtual FString VolumeInfo() const override;
-		virtual void GetActiveVoxels(TArray<FBox>& OutActiveVoxels) const override;
+		DATAFLOWVOLUMECORE_API void GetActiveVoxels(const float InIsovalue, TArray<FBox>& OutActiveVoxels, bool bInteriorMaskOnly = false) const;
 		virtual FBox GetVolumeBoundingBox() const override;
 		virtual bool CreateVolumeTexture(UVolumeTexture* InVolumeTexture) const override;
 
diff --git a/Engine/Plugins/Experimental/DataflowVolumetricPlugin/Source/DataflowVolumeCore/Private/Dataflow/DataflowVolume.cpp b/Engine/Plugins/Experimental/DataflowVolumetricPlugin/Source/DataflowVolumeCore/Private/Dataflow/DataflowVolume.cpp
--- a/Engine/Plugins/Experimental/DataflowVolumetricPlugin/Source/DataflowVolumeCore/Private/Dataflow/DataflowVolume.cpp
+++ b/Engine/Plugins/Experimental/DataflowVolumetricPlugin/Source/DataflowVolumeCore/Private/Dataflow/DataflowVolume.cpp
@@ -26,12 +26,12 @@ FString FDataflowVolume::VolumeInfo() const
 
 /* --------------------------------------------------------------------------------------------------------------- */
 
-int32 FDataflowVolume::GetNumActiveVoxels() const
+int32 FDataflowVolume::GetNumActiveVoxels(const float InIsovalue, bool bInteriorMaskOnly) const
 {
 	if (Volume)
 	{
 		TArray<FBox> OutActiveVoxels;
-		Volume->GetActiveVoxels(OutActiveVoxels);
+		Volume->GetActiveVoxels(InIsovalue, OutActiveVoxels, bInteriorMaskOnly);
 
 		return OutActiveVoxels.Num();
 	}
@@ -41,11 +41,11 @@ int32 FDataflowVolume::GetNumActiveVoxels() const
 
 /* --------------------------------------------------------------------------------------------------------------- */
 
-void FDataflowVolume::GetActiveVoxels(TArray<FBox>& OutActiveVoxels) const
+void FDataflowVolume::GetActiveVoxels(const float InIsovalue, TArray<FBox>& OutActiveVoxels, bool bInteriorMaskOnly) const
 {
 	if (Volume)
 	{
-		Volume->GetActiveVoxels(OutActiveVoxels);
+		Volume->GetActiveVoxels(InIsovalue, OutActiveVoxels, bInteriorMaskOnly);
 	}
 }
 
diff --git a/Engine/Plugins/Experimental/DataflowVolumetricPlugin/Source/DataflowVolumeCore/Private/Dataflow/DataflowVolumeImpl.h b/Engine/Plugins/Experimental/DataflowVolumetricPlugin/Source/DataflowVolumeCore/Private/Dataflow/DataflowVolumeImpl.h
--- a/Engine/Plugins/Experimental/DataflowVolumetricPlugin/Source/DataflowVolumeCore/Private/Dataflow/DataflowVolumeImpl.h
+++ b/Engine/Plugins/Experimental/DataflowVolumetricPlugin/Source/DataflowVolumeCore/Private/Dataflow/DataflowVolumeImpl.h
@@ -27,7 +27,7 @@ namespace UE::DataflowVolume::Private
 
 		int32 GetActiveVoxelCount() const;
 
-		virtual void GetActiveVoxels(TArray<FBox>& OutActiveVoxels) const = 0;
+		virtual void GetActiveVoxels(const float InIsovalue, TArray<FBox>& OutActiveVoxels, bool bInteriorMaskOnly = false) const = 0;
 
 		virtual FBox GetVolumeBoundingBox() const = 0;
 
diff --git a/Engine/Plugins/Experimental/DataflowVolumetricPlugin/Source/DataflowVolumeCore/Private/Dataflow/DataflowVolumeUtils.cpp b/Engine/Plugins/Experimental/DataflowVolumetricPlugin/Source/DataflowVolumeCore/Private/Dataflow/DataflowVolumeUtils.cpp
--- a/Engine/Plugins/Experimental/DataflowVolumetricPlugin/Source/DataflowVolumeCore/Private/Dataflow/DataflowVolumeUtils.cpp
+++ b/Engine/Plugins/Experimental/DataflowVolumetricPlugin/Source/DataflowVolumeCore/Private/Dataflow/DataflowVolumeUtils.cpp
@@ -196,8 +196,7 @@ namespace UE::DataflowVolumeUtils
 		return VDBMatDouble;
 	}
 
-
-
+	/* ---------------------------------------------------------------------------------------------------- */
 
 }
 
diff --git a/Engine/Plugins/Experimental/DataflowVolumetricPlugin/Source/DataflowVolumeCore/Private/Dataflow/DataflowVolumeUtils.h b/Engine/Plugins/Experimental/DataflowVolumetricPlugin/Source/DataflowVolumeCore/Private/Dataflow/DataflowVolumeUtils.h
--- a/Engine/Plugins/Experimental/DataflowVolumetricPlugin/Source/DataflowVolumeCore/Private/Dataflow/DataflowVolumeUtils.h
+++ b/Engine/Plugins/Experimental/DataflowVolumetricPlugin/Source/DataflowVolumeCore/Private/Dataflow/DataflowVolumeUtils.h
@@ -697,4 +697,67 @@ namespace UE::DataflowVolumeUtils
 	* Convert a FMatrix to openvdb::Mat4R
 	*/
 	openvdb::Mat4R ConvertMatrixToVDBMatrix(FMatrix& InMatrix);
+
+	template <typename GridType>
+	void GetActiveVoxels(const typename GridType::Ptr& InGridPtr, TArray<FBox>& OutActiveVoxels)
+	{
+		using IterT = typename GridType::ValueOnCIter;
+
+		openvdb::math::Transform Transform = InGridPtr->transform();
+
+		if (!InGridPtr->tree().empty())
+		{
+			for (IterT iter = InGridPtr->cbeginValueOn(); iter.test(); ++iter)
+			{
+				if (iter.isVoxelValue())
+				{
+					openvdb::CoordBBox BBox;
+					iter.getBoundingBox(BBox);
+
+					const openvdb::Vec3d BBoxMin(BBox.min().x() - 0.5, BBox.min().y() - 0.5, BBox.min().z() - 0.5);
+					const openvdb::Vec3d BBoxMax(BBox.max().x() + 0.5, BBox.max().y() + 0.5, BBox.max().z() + 0.5);
+
+					const openvdb::Vec3d BBoxMinWS = Transform.indexToWorld(BBoxMin);
+					const openvdb::Vec3d BBoxMaxWS = Transform.indexToWorld(BBoxMax);
+
+					FVector Min(BBoxMinWS.x(), BBoxMinWS.y(), BBoxMinWS.z());
+					FVector Max(BBoxMaxWS.x(), BBoxMaxWS.y(), BBoxMaxWS.z());
+
+					OutActiveVoxels.Add(FBox(Min, Max));
+				}
+			}
+		}
+	}
+
+	/**
+	* Get active voxels of a volume of any type
+	*/
+	template <typename GridType>
+	void GetActiveVoxels(const typename GridType::Ptr& InGridPtr, const float InIsovalue, TArray<FBox>& OutActiveVoxels, bool bInteriorMaskOnly = false)
+	{
+		if (!InGridPtr->tree().empty())
+		{
+			if (bInteriorMaskOnly)
+			{
+				// Compute a mask of the voxels enclosed by the isosurface
+				auto interiorMaskPtr = [&InGridPtr, &InIsovalue]
+					{
+						if constexpr (std::is_same_v<GridType, openvdb::FloatGrid>)
+						{
+							return openvdb::tools::sdfInteriorMask(*InGridPtr, InIsovalue);
+						}
+						else
+						{
+							return openvdb::tools::interiorMask(*InGridPtr, InIsovalue);
+						}
+					} ();
+
+				GetActiveVoxels<openvdb::BoolGrid>(interiorMaskPtr, OutActiveVoxels);
+			}
+			else
+			{
+				GetActiveVoxels<GridType>(InGridPtr, OutActiveVoxels);
+			}
+		}
+	}
 }
diff --git a/Engine/Plugins/Experimental/DataflowVolumetricPlugin/Source/DataflowVolumeCore/Private/Dataflow/OpenVDB.h b/Engine/Plugins/Experimental/DataflowVolumetricPlugin/Source/DataflowVolumeCore/Private/Dataflow/OpenVDB.h
--- a/Engine/Plugins/Experimental/DataflowVolumetricPlugin/Source/DataflowVolumeCore/Private/Dataflow/OpenVDB.h
+++ b/Engine/Plugins/Experimental/DataflowVolumetricPlugin/Source/DataflowVolumeCore/Private/Dataflow/OpenVDB.h
@@ -56,7 +56,7 @@ UE_PUSH_MACRO("check")
 #include <openvdb/tools/LevelSetFracture.h>
 #include <openvdb/tools/ParticlesToLevelSet.h>
 #include <openvdb/tools/GridOperators.h>
-
+#include <openvdb/tools/Mask.h>
 #if LOCAL_M_PI
 	#undef M_PI
 #endif
diff --git a/Engine/Plugins/Experimental/DataflowVolumetricPlugin/Source/DataflowVolumeNodes/Private/Dataflow/DataflowVolumeNodes.cpp b/Engine/Plugins/Experimental/DataflowVolumetricPlugin/Source/DataflowVolumeNodes/Private/Dataflow/DataflowVolumeNodes.cpp
--- a/Engine/Plugins/Experimental/DataflowVolumetricPlugin/Source/DataflowVolumeNodes/Private/Dataflow/DataflowVolumeNodes.cpp
+++ b/Engine/Plugins/Experimental/DataflowVolumetricPlugin/Source/DataflowVolumeNodes/Private/Dataflow/DataflowVolumeNodes.cpp
@@ -82,59 +82,6 @@ void FMakeSDFSphereDataflowNode::Evaluate(UE::Dataflow::FContext& Context, const
 	}
 }
 
-#if WITH_EDITOR
-bool FMakeSDFSphereDataflowNode::CanDebugDrawViewMode(const FName& ViewModeName) const
-{
-	return ViewModeName == UE::Dataflow::FDataflowConstruction3DViewMode::Name;
-}
-
-void FMakeSDFSphereDataflowNode::DebugDraw(UE::Dataflow::FContext& Context, IDataflowDebugDrawInterface& DataflowRenderingInterface, const FDebugDrawParameters& DebugDrawParameters) const
-{
-	using namespace UE::Dataflow::Overlay;
-
-	if (DebugDrawParameters.bNodeIsSelected || DebugDrawParameters.bNodeIsPinned)
-	{
-		FString OverlayEvaluateResultString;
-
-		const FDataflowOutput* FloatVolumeOutput = FindOutput(&FloatVolume);
-		if (FloatVolumeOutput)
-		{
-			const FDataflowFloatVolume& OutFloatVolume = FloatVolumeOutput->GetValue(Context, FloatVolume);
-//			const FDataflowFloatVolume& OutFloatVolume = this->GetOutputValue(Context, FloatVolume);
-
-			DebugDrawRenderSettings.SetDebugDrawSettings(DataflowRenderingInterface);
-
-			TArray<FBox> ActiveVoxels;
-			OutFloatVolume.GetActiveVoxels(ActiveVoxels);
-			
-			for (const FBox& Voxel : ActiveVoxels)
-			{
-				if (bDrawVoxelCenters)
-				{
-					DataflowRenderingInterface.DrawPoint(0.5 * FVector(Voxel.Max + Voxel.Min));
-				}
-
-				if (bDrawVoxels)
-				{
-					DataflowRenderingInterface.DrawBox(0.5 * FVector(Voxel.Max - Voxel.Min), FQuat::Identity, 0.5 * FVector(Voxel.Min + Voxel.Max), 1.0);
-				}
-			}
-
-			OverlayEvaluateResultString = OutFloatVolume.VolumeInfo();
-		}
-
-		// ------------------------------------------------------------------------------------------------------
-		// Overlay
-		TArray<FString> OverlayStrArr;
-		OverlayStrArr.Add(BuildOverlayNodeInfoString(this));
-		OverlayStrArr.Add(OverlayEvaluateResultString);
-		FString OverlayStr = BuildOverlayFinalString(OverlayStrArr);
-		DataflowRenderingInterface.DrawOverlayText(OverlayStr);
-		// ------------------------------------------------------------------------------------------------------
-	}
-}
-#endif
-
 /* --------------------------------------------------------------------------------------------- */
 
 FMakeSDFCubeDataflowNode::FMakeSDFCubeDataflowNode(const UE::Dataflow::FNodeParameters& InParam, FGuid InGuid)
@@ -168,59 +115,6 @@ void FMakeSDFCubeDataflowNode::Evaluate(UE::Dataflow::FContext& Context, const F
 	}
 }
 
-#if WITH_EDITOR
-bool FMakeSDFCubeDataflowNode::CanDebugDrawViewMode(const FName& ViewModeName) const
-{
-	return ViewModeName == UE::Dataflow::FDataflowConstruction3DViewMode::Name;
-}
-
-void FMakeSDFCubeDataflowNode::DebugDraw(UE::Dataflow::FContext& Context, IDataflowDebugDrawInterface& DataflowRenderingInterface, const FDebugDrawParameters& DebugDrawParameters) const
-{
-	using namespace UE::Dataflow::Overlay;
-
-	if (DebugDrawParameters.bNodeIsSelected || DebugDrawParameters.bNodeIsPinned)
-	{
-		FString OverlayEvaluateResultString;
-
-		const FDataflowOutput* FloatVolumeOutput = FindOutput(&FloatVolume);
-		if (FloatVolumeOutput)
-		{
-			const FDataflowFloatVolume& OutFloatVolume = FloatVolumeOutput->GetValue(Context, FloatVolume);
-			//			const FDataflowFloatVolume& OutFloatVolume = this->GetOutputValue(Context, FloatVolume);
-
-			DebugDrawRenderSettings.SetDebugDrawSettings(DataflowRenderingInterface);
-
-			TArray<FBox> ActiveVoxels;
-			OutFloatVolume.GetActiveVoxels(ActiveVoxels);
-
-			for (const FBox& Voxel : ActiveVoxels)
-			{
-				if (bDrawVoxelCenters)
-				{
-					DataflowRenderingInterface.DrawPoint(0.5 * FVector(Voxel.Max + Voxel.Min));
-				}
-
-				if (bDrawVoxels)
-				{
-					DataflowRenderingInterface.DrawBox(0.5 * FVector(Voxel.Max - Voxel.Min), FQuat::Identity, 0.5 * FVector(Voxel.Min + Voxel.Max), 1.0);
-				}
-			}
-
-			OverlayEvaluateResultString = OutFloatVolume.VolumeInfo();
-		}
-
-		// ------------------------------------------------------------------------------------------------------
-		// Overlay
-		TArray<FString> OverlayStrArr;
-		OverlayStrArr.Add(BuildOverlayNodeInfoString(this));
-		OverlayStrArr.Add(OverlayEvaluateResultString);
-		FString OverlayStr = BuildOverlayFinalString(OverlayStrArr);
-		DataflowRenderingInterface.DrawOverlayText(OverlayStr);
-		// ------------------------------------------------------------------------------------------------------
-	}
-}
-#endif
-
 /* --------------------------------------------------------------------------------------------- */
 
 FMakeSDFPlatonicSolidDataflowNode::FMakeSDFPlatonicSolidDataflowNode(const UE::Dataflow::FNodeParameters& InParam, FGuid InGuid)
@@ -277,59 +171,6 @@ void FMakeSDFPlatonicSolidDataflowNode::Evaluate(UE::Dataflow::FContext& Context
 	}
 }
 
-#if WITH_EDITOR
-bool FMakeSDFPlatonicSolidDataflowNode::CanDebugDrawViewMode(const FName& ViewModeName) const
-{
-	return ViewModeName == UE::Dataflow::FDataflowConstruction3DViewMode::Name;
-}
-
-void FMakeSDFPlatonicSolidDataflowNode::DebugDraw(UE::Dataflow::FContext& Context, IDataflowDebugDrawInterface& DataflowRenderingInterface, const FDebugDrawParameters& DebugDrawParameters) const
-{
-	using namespace UE::Dataflow::Overlay;
-
-	if (DebugDrawParameters.bNodeIsSelected || DebugDrawParameters.bNodeIsPinned)
-	{
-		FString OverlayEvaluateResultString;
-
-		const FDataflowOutput* FloatVolumeOutput = FindOutput(&FloatVolume);
-		if (FloatVolumeOutput)
-		{
-			const FDataflowFloatVolume& OutFloatVolume = FloatVolumeOutput->GetValue(Context, FloatVolume);
-			//			const FDataflowFloatVolume& OutFloatVolume = this->GetOutputValue(Context, FloatVolume);
-
-			DebugDrawRenderSettings.SetDebugDrawSettings(DataflowRenderingInterface);
-
-			TArray<FBox> ActiveVoxels;
-			OutFloatVolume.GetActiveVoxels(ActiveVoxels);
-
-			for (const FBox& Voxel : ActiveVoxels)
-			{
-				if (bDrawVoxelCenters)
-				{
-					DataflowRenderingInterface.DrawPoint(0.5 * FVector(Voxel.Max + Voxel.Min));
-				}
-
-				if (bDrawVoxels)
-				{
-					DataflowRenderingInterface.DrawBox(0.5 * FVector(Voxel.Max - Voxel.Min), FQuat::Identity, 0.5 * FVector(Voxel.Min + Voxel.Max), 1.0);
-				}
-			}
-
-			OverlayEvaluateResultString = OutFloatVolume.VolumeInfo();
-		}
-
-		// ------------------------------------------------------------------------------------------------------
-		// Overlay
-		TArray<FString> OverlayStrArr;
-		OverlayStrArr.Add(BuildOverlayNodeInfoString(this));
-		OverlayStrArr.Add(OverlayEvaluateResultString);
-		FString OverlayStr = BuildOverlayFinalString(OverlayStrArr);
-		DataflowRenderingInterface.DrawOverlayText(OverlayStr);
-		// ------------------------------------------------------------------------------------------------------
-	}
-}
-#endif
-
 /* --------------------------------------------------------------------------------------------- */
 
 FVolumeInfoDataflowNode::FVolumeInfoDataflowNode(const UE::Dataflow::FNodeParameters& InParam, FGuid InGuid)
@@ -436,59 +277,6 @@ void FCollectionToVolumeDataflowNode::Evaluate(UE::Dataflow::FContext& Context,
 	}
 }
 
-#if WITH_EDITOR
-bool FCollectionToVolumeDataflowNode::CanDebugDrawViewMode(const FName& ViewModeName) const
-{
-	return ViewModeName == UE::Dataflow::FDataflowConstruction3DViewMode::Name;
-}
-
-void FCollectionToVolumeDataflowNode::DebugDraw(UE::Dataflow::FContext& Context, IDataflowDebugDrawInterface& DataflowRenderingInterface, const FDebugDrawParameters& DebugDrawParameters) const
-{
-	using namespace UE::Dataflow::Overlay;
-
-	if (DebugDrawParameters.bNodeIsSelected || DebugDrawParameters.bNodeIsPinned)
-	{
-		FString OverlayEvaluateResultString;
-
-		const FDataflowOutput* FloatVolumeOutput = FindOutput(&FloatVolume);
-		if (FloatVolumeOutput)
-		{
-			const FDataflowFloatVolume& OutFloatVolume = FloatVolumeOutput->GetValue(Context, FloatVolume);
-			//			const FDataflowFloatVolume& OutFloatVolume = this->GetOutputValue(Context, FloatVolume);
-
-			DebugDrawRenderSettings.SetDebugDrawSettings(DataflowRenderingInterface);
-
-			TArray<FBox> ActiveVoxels;
-			OutFloatVolume.GetActiveVoxels(ActiveVoxels);
-
-			for (const FBox& Voxel : ActiveVoxels)
-			{
-				if (bDrawVoxelCenters)
-				{
-					DataflowRenderingInterface.DrawPoint(0.5 * FVector(Voxel.Max + Voxel.Min));
-				}
-
-				if (bDrawVoxels)
-				{
-					DataflowRenderingInterface.DrawBox(0.5 * FVector(Voxel.Max - Voxel.Min), FQuat::Identity, 0.5 * FVector(Voxel.Min + Voxel.Max), 1.0);
-				}
-			}
-
-			OverlayEvaluateResultString = OutFloatVolume.VolumeInfo();
-		}
-
-		// ------------------------------------------------------------------------------------------------------
-		// Overlay
-		TArray<FString> OverlayStrArr;
-		OverlayStrArr.Add(BuildOverlayNodeInfoString(this));
-		OverlayStrArr.Add(OverlayEvaluateResultString);
-		FString OverlayStr = BuildOverlayFinalString(OverlayStrArr);
-		DataflowRenderingInterface.DrawOverlayText(OverlayStr);
-		// ------------------------------------------------------------------------------------------------------
-	}
-}
-#endif
-
 /* --------------------------------------------------------------------------------------------- */
 
 FVolumeToSpheresDataflowNode::FVolumeToSpheresDataflowNode(const UE::Dataflow::FNodeParameters& InParam, FGuid InGuid)
@@ -503,11 +291,6 @@ FVolumeToSpheresDataflowNode::FVolumeToSpheresDataflowNode(const UE::Dataflow::F
 	RegisterInputConnection(&InstanceCount).SetCanHidePin(true).SetPinIsHidden(true);
 	RegisterOutputConnection(&Collection);
 	RegisterOutputConnection(&Spheres);
-
-	DebugDrawRenderSettings.RenderType = EDataflowDebugDrawRenderType::Wireframe;
-	DebugDrawRenderSettings.bTranslucent = false;
-	DebugDrawRenderSettings.Color = FLinearColor::Blue;
-	DebugDrawRenderSettings.LineWidthMultiplier = 1.5f;
 }
 
 void FVolumeToSpheresDataflowNode::Evaluate(UE::Dataflow::FContext& Context, const FDataflowOutput* Out) const
@@ -598,75 +381,6 @@ void FVolumeToSpheresDataflowNode::Evaluate(UE::Dataflow::FContext& Context, con
 	}
 }
 
-#if WITH_EDITOR
-bool FVolumeToSpheresDataflowNode::CanDebugDrawViewMode(const FName& ViewModeName) const
-{
-	return ViewModeName == UE::Dataflow::FDataflowConstruction3DViewMode::Name;
-}
-
-void FVolumeToSpheresDataflowNode::DebugDraw(UE::Dataflow::FContext& Context, IDataflowDebugDrawInterface& DataflowRenderingInterface, const FDebugDrawParameters& DebugDrawParameters) const
-{
-	using namespace UE::Dataflow::Overlay;
-
-	if (DebugDrawParameters.bNodeIsSelected || DebugDrawParameters.bNodeIsPinned)
-	{
-		FString OverlayEvaluateResultString;
-
-		const FDataflowOutput* SpheresOutput = FindOutput(&Spheres);
-		if (SpheresOutput)
-		{
-			const TArray<FSphere>& OutSpheres = SpheresOutput->GetValue(Context, Spheres);
-			//			const FDataflowFloatVolume& OutFloatVolume = this->GetOutputValue(Context, FloatVolume);
-
-			DebugDrawRenderSettings.SetDebugDrawSettings(DataflowRenderingInterface);
-
-			float MinRad = FLT_MAX, MaxRad = -FLT_MAX, SumRad = 0;
-			for (const FSphere& Sphere : OutSpheres)
-			{
-				if (Sphere.W < MinRad)
-				{
-					MinRad = Sphere.W;
-				}
-				if (Sphere.W > MaxRad)
-				{
-					MaxRad = Sphere.W;
-				}
-
-				SumRad += Sphere.W;
-
-				if (bDrawSphereCenters)
-				{
-					DataflowRenderingInterface.DrawPoint(Sphere.Center);
-				}
-
-				if (bDrawSpheres)
-				{
-					DataflowRenderingInterface.DrawSphere(Sphere.Center, Sphere.W);
-				}
-			}
-
-			const int32 NumSpheres = OutSpheres.Num();
-			const float MeanRad = SumRad / (float)NumSpheres;
-
-			OverlayEvaluateResultString.Appendf(TEXT("Number of Spheres: %d\nMin radius: %3.3f\nMax radius: %3.3f\nMedian radius: %3.3f\n"),
-				NumSpheres,
-				MinRad,
-				MaxRad,
-				MeanRad);
-		}
-
-		// ------------------------------------------------------------------------------------------------------
-		// Overlay
-		TArray<FString> OverlayStrArr;
-		OverlayStrArr.Add(BuildOverlayNodeInfoString(this));
-		OverlayStrArr.Add(OverlayEvaluateResultString);
-		FString OverlayStr = BuildOverlayFinalString(OverlayStrArr);
-		DataflowRenderingInterface.DrawOverlayText(OverlayStr);
-		// ------------------------------------------------------------------------------------------------------
-	}
-}
-#endif
-
 /* --------------------------------------------------------------------------------------------- */
 
 FVolumeToSpheresDataflowNode_v2::FVolumeToSpheresDataflowNode_v2(const UE::Dataflow::FNodeParameters& InParam, FGuid InGuid)
@@ -891,67 +605,6 @@ void FConvertVolumeDataflowNode::Evaluate(UE::Dataflow::FContext& Context, const
 	}
 }
 
-#if WITH_EDITOR
-bool FConvertVolumeDataflowNode::CanDebugDrawViewMode(const FName& ViewModeName) const
-{
-	return ViewModeName == UE::Dataflow::FDataflowConstruction3DViewMode::Name;
-}
-
-void FConvertVolumeDataflowNode::DebugDraw(UE::Dataflow::FContext& Context, IDataflowDebugDrawInterface& DataflowRenderingInterface, const FDebugDrawParameters& DebugDrawParameters) const
-{
-	using namespace UE::Dataflow::Overlay;
-
-	if (DebugDrawParameters.bNodeIsSelected || DebugDrawParameters.bNodeIsPinned)
-	{
-		if (ConvertTo == EDataflowVolumeConvertSDFTo::Volume)
-		{
-			FString OverlayEvaluateResultString;
-
-			const FDataflowOutput* FloatVolumeOutput = FindOutput(&FloatVolume);
-			if (FloatVolumeOutput)
-			{
-				const FDataflowFloatVolume& OutFloatVolume = FloatVolumeOutput->GetValue(Context, FloatVolume);
-				//			const FDataflowFloatVolume& OutFloatVolume = this->GetOutputValue(Context, FloatVolume);
-
-				// Check if grid is empty
-//				if (!OutFloatVolume.GetFloatVolume()->IsGridEmpty())
-				if (0)
-				{
-					DebugDrawRenderSettings.SetDebugDrawSettings(DataflowRenderingInterface);
-
-					TArray<FBox> ActiveVoxels;
-					OutFloatVolume.GetActiveVoxels(ActiveVoxels);
-
-					for (const FBox& Voxel : ActiveVoxels)
-					{
-						if (bDrawVoxelCenters)
-						{
-							DataflowRenderingInterface.DrawPoint(0.5 * FVector(Voxel.Max + Voxel.Min));
-						}
-
-						if (bDrawVoxels)
-						{
-							DataflowRenderingInterface.DrawBox(0.5 * FVector(Voxel.Max - Voxel.Min), FQuat::Identity, 0.5 * FVector(Voxel.Min + Voxel.Max), 1.0);
-						}
-					}
-
-					OverlayEvaluateResultString = OutFloatVolume.VolumeInfo();
-
-					// ------------------------------------------------------------------------------------------------------
-					// Overlay
-					TArray<FString> OverlayStrArr;
-					OverlayStrArr.Add(BuildOverlayNodeInfoString(this));
-					OverlayStrArr.Add(OverlayEvaluateResultString);
-					FString OverlayStr = BuildOverlayFinalString(OverlayStrArr);
-					DataflowRenderingInterface.DrawOverlayText(OverlayStr);
-					// ------------------------------------------------------------------------------------------------------
-				}
-			}
-		}
-	}
-}
-#endif
-
 /* --------------------------------------------------------------------------------------------- */
 
 FUniformVolumeScatterDataflowNode::FUniformVolumeScatterDataflowNode(const UE::Dataflow::FNodeParameters& InParam, FGuid InGuid)
@@ -999,49 +652,6 @@ void FUniformVolumeScatterDataflowNode::Evaluate(UE::Dataflow::FContext& Context
 	}
 }
 
-#if WITH_EDITOR
-bool FUniformVolumeScatterDataflowNode::CanDebugDrawViewMode(const FName& ViewModeName) const
-{
-	return ViewModeName == UE::Dataflow::FDataflowConstruction3DViewMode::Name;
-}
-
-void FUniformVolumeScatterDataflowNode::DebugDraw(UE::Dataflow::FContext& Context, IDataflowDebugDrawInterface& DataflowRenderingInterface, const FDebugDrawParameters& DebugDrawParameters) const
-{
-	using namespace UE::Dataflow::Overlay;
-
-	if (DebugDrawParameters.bNodeIsSelected || DebugDrawParameters.bNodeIsPinned)
-	{
-		FString OverlayEvaluateResultString;
-
-		const FDataflowOutput* PointsOutput = FindOutput(&Points);
-		if (PointsOutput)
-		{
-			const TArray<FVector>& OutPoints = PointsOutput->GetValue(Context, Points);
-			//			const FDataflowFloatVolume& OutFloatVolume = this->GetOutputValue(Context, FloatVolume);
-
-			DataflowRenderingInterface.SetWorldPriority();
-			DataflowRenderingInterface.SetColor(Color);
-
-			for (const FVector& Point : OutPoints)
-			{
-				DataflowRenderingInterface.DrawPoint(Point);
-			}
-
-			OverlayEvaluateResultString.Appendf(TEXT("Number of Points: %d\n"), OutPoints.Num());
-		}
-
-		// ------------------------------------------------------------------------------------------------------
-		// Overlay
-		TArray<FString> OverlayStrArr;
-		OverlayStrArr.Add(BuildOverlayNodeInfoString(this));
-		OverlayStrArr.Add(OverlayEvaluateResultString);
-		FString OverlayStr = BuildOverlayFinalString(OverlayStrArr);
-		DataflowRenderingInterface.DrawOverlayText(OverlayStr);
-		// ------------------------------------------------------------------------------------------------------
-	}
-}
-#endif
-
 /* --------------------------------------------------------------------------------------------- */
 
 FDenseUniformVolumeScatterDataflowNode::FDenseUniformVolumeScatterDataflowNode(const UE::Dataflow::FNodeParameters& InParam, FGuid InGuid)
@@ -1089,49 +699,6 @@ void FDenseUniformVolumeScatterDataflowNode::Evaluate(UE::Dataflow::FContext& Co
 	}
 }
 
-#if WITH_EDITOR
-bool FDenseUniformVolumeScatterDataflowNode::CanDebugDrawViewMode(const FName& ViewModeName) const
-{
-	return ViewModeName == UE::Dataflow::FDataflowConstruction3DViewMode::Name;
-}
-
-void FDenseUniformVolumeScatterDataflowNode::DebugDraw(UE::Dataflow::FContext& Context, IDataflowDebugDrawInterface& DataflowRenderingInterface, const FDebugDrawParameters& DebugDrawParameters) const
-{
-	using namespace UE::Dataflow::Overlay;
-
-	if (DebugDrawParameters.bNodeIsSelected || DebugDrawParameters.bNodeIsPinned)
-	{
-		FString OverlayEvaluateResultString;
-
-		const FDataflowOutput* PointsOutput = FindOutput(&Points);
-		if (PointsOutput)
-		{
-			const TArray<FVector>& OutPoints = PointsOutput->GetValue(Context, Points);
-			//			const FDataflowFloatVolume& OutFloatVolume = this->GetOutputValue(Context, FloatVolume);
-
-			DataflowRenderingInterface.SetWorldPriority();
-			DataflowRenderingInterface.SetColor(Color);
-
-			for (const FVector& Point : OutPoints)
-			{
-				DataflowRenderingInterface.DrawPoint(Point);
-			}
-
-			OverlayEvaluateResultString.Appendf(TEXT("Number of Points: %d\n"), OutPoints.Num());
-		}
-
-		// ------------------------------------------------------------------------------------------------------
-		// Overlay
-		TArray<FString> OverlayStrArr;
-		OverlayStrArr.Add(BuildOverlayNodeInfoString(this));
-		OverlayStrArr.Add(OverlayEvaluateResultString);
-		FString OverlayStr = BuildOverlayFinalString(OverlayStrArr);
-		DataflowRenderingInterface.DrawOverlayText(OverlayStr);
-		// ------------------------------------------------------------------------------------------------------
-	}
-}
-#endif
-
 /* --------------------------------------------------------------------------------------------- */
 
 FNonUniformVolumeScatterDataflowNode::FNonUniformVolumeScatterDataflowNode(const UE::Dataflow::FNodeParameters& InParam, FGuid InGuid)
@@ -1179,49 +746,6 @@ void FNonUniformVolumeScatterDataflowNode::Evaluate(UE::Dataflow::FContext& Cont
 	}
 }
 
-#if WITH_EDITOR
-bool FNonUniformVolumeScatterDataflowNode::CanDebugDrawViewMode(const FName& ViewModeName) const
-{
-	return ViewModeName == UE::Dataflow::FDataflowConstruction3DViewMode::Name;
-}
-
-void FNonUniformVolumeScatterDataflowNode::DebugDraw(UE::Dataflow::FContext& Context, IDataflowDebugDrawInterface& DataflowRenderingInterface, const FDebugDrawParameters& DebugDrawParameters) const
-{
-	using namespace UE::Dataflow::Overlay;
-
-	if (DebugDrawParameters.bNodeIsSelected || DebugDrawParameters.bNodeIsPinned)
-	{
-		FString OverlayEvaluateResultString;
-
-		const FDataflowOutput* PointsOutput = FindOutput(&Points);
-		if (PointsOutput)
-		{
-			const TArray<FVector>& OutPoints = PointsOutput->GetValue(Context, Points);
-			//			const FDataflowFloatVolume& OutFloatVolume = this->GetOutputValue(Context, FloatVolume);
-
-			DataflowRenderingInterface.SetWorldPriority();
-			DataflowRenderingInterface.SetColor(Color);
-
-			for (const FVector& Point : OutPoints)
-			{
-				DataflowRenderingInterface.DrawPoint(Point);
-			}
-
-			OverlayEvaluateResultString.Appendf(TEXT("Number of Points: %d\n"), OutPoints.Num());
-		}
-
-		// ------------------------------------------------------------------------------------------------------
-		// Overlay
-		TArray<FString> OverlayStrArr;
-		OverlayStrArr.Add(BuildOverlayNodeInfoString(this));
-		OverlayStrArr.Add(OverlayEvaluateResultString);
-		FString OverlayStr = BuildOverlayFinalString(OverlayStrArr);
-		DataflowRenderingInterface.DrawOverlayText(OverlayStr);
-		// ------------------------------------------------------------------------------------------------------
-	}
-}
-#endif
-
 /* --------------------------------------------------------------------------------------------- */
 
 FVolumeSampleToPointsDataflowNode::FVolumeSampleToPointsDataflowNode(const UE::Dataflow::FNodeParameters& InParam, FGuid InGuid)
@@ -1277,7 +801,7 @@ void FVolumeSampleToPointsDataflowNode::OnPropertyChanged(UE::Dataflow::FContext
 
 void FVolumeSampleToPointsDataflowNode::Evaluate(UE::Dataflow::FContext& Context, const FDataflowOutput* Out) const
 {
-	if (Out->IsA(&OutArray))
+	if (Out->IsA(&OutArray) || Out->IsA(&PointsCollection))
 	{
 		const FDataflowOutput* OutArrayOutput = FindOutput(&OutArray);
 
@@ -1382,48 +906,23 @@ void FVolumeSampleToPointsDataflowNode::Evaluate(UE::Dataflow::FContext& Context
 
 		if (OutArrayOutput)
 		{
-			OutArrayOutput->SetValue<TArray<FVector4>>(TArray<FVector4>(), Context);
-		}
-	}
-}
-
-#if WITH_EDITOR
-bool FVolumeSampleToPointsDataflowNode::CanDebugDrawViewMode(const FName& ViewModeName) const
-{
-	return ViewModeName == UE::Dataflow::FDataflowConstruction3DViewMode::Name;
-}
-
-void FVolumeSampleToPointsDataflowNode::DebugDraw(UE::Dataflow::FContext& Context, IDataflowDebugDrawInterface& DataflowRenderingInterface, const FDebugDrawParameters& DebugDrawParameters) const
-{
-	using namespace UE::Dataflow::Overlay;
-
-	if (DebugDrawParameters.bNodeIsSelected || DebugDrawParameters.bNodeIsPinned)
-	{
-		FString OverlayEvaluateResultString;
-
-		const TArray<FVector4>& InSamplePoints = GetValue(Context, &SamplePoints);
-
-		DataflowRenderingInterface.SetWorldPriority();
-		DataflowRenderingInterface.SetColor(Color);
+			if (SampleType == EDataflowVolumeSampleType::Float)
+			{
+				OutArrayOutput->SetValue<TArray<float>>(TArray<float>(), Context);
+			}
+			else if (SampleType == EDataflowVolumeSampleType::Int)
+			{
+				OutArrayOutput->SetValue<TArray<int32>>(TArray<int32>(), Context);
+			}
+			else if (SampleType == EDataflowVolumeSampleType::FloatVector)
+			{
+				OutArrayOutput->SetValue<TArray<FVector>>(TArray<FVector>(), Context);
+			}
 
-		for (const FVector4& Point : InSamplePoints)
-		{
-			DataflowRenderingInterface.DrawPoint(FVector(Point));
+			SetValue(Context, FManagedArrayCollection(), &PointsCollection);
 		}
-
-		OverlayEvaluateResultString.Appendf(TEXT("Number of Points: %d\n"), InSamplePoints.Num());
-
-		// ------------------------------------------------------------------------------------------------------
-		// Overlay
-		TArray<FString> OverlayStrArr;
-		OverlayStrArr.Add(BuildOverlayNodeInfoString(this));
-		OverlayStrArr.Add(OverlayEvaluateResultString);
-		FString OverlayStr = BuildOverlayFinalString(OverlayStrArr);
-		DataflowRenderingInterface.DrawOverlayText(OverlayStr);
-		// ------------------------------------------------------------------------------------------------------
 	}
 }
-#endif
 
 /* --------------------------------------------------------------------------------------------- */
 
@@ -1522,20 +1021,6 @@ namespace UE::Dataflow::Volume
 		return Points;
 	}
 
-	//static void SetColorBar(const TArray<FLinearColor>& Spectrum, const FDataflowColorRamp ColorRamp)
-	//{
-	//	const int32 NumColors = Spectrum.Num();
-
-	//	for (int32 Idx = 0; Idx < NumColors; ++Idx)
-	//	{
-	//		ColorRamp.ColorRamp.GetCurves()[0].CurveToEdit->AddKey((float)Idx * 1.f / ((float)NumColors - 1.f), Spectrum[Idx].R);
-	//		ColorRamp.ColorRamp.GetCurves()[1].CurveToEdit->AddKey((float)Idx * 1.f / ((float)NumColors - 1.f), Spectrum[Idx].G);
-	//		ColorRamp.ColorRamp.GetCurves()[2].CurveToEdit->AddKey((float)Idx * 1.f / ((float)NumColors - 1.f), Spectrum[Idx].B);
-	//	}
-	//	ColorRamp.ColorRamp.GetCurves()[3].CurveToEdit->AddKey(0.f, 1.f);
-
-	//}
-
 	static TArray< FLinearColor> GetInfraRedColorSpectrum()
 	{
 		const int32 NumColors = 7;
@@ -1836,11 +1321,11 @@ void FVolumeAnalysisDataflowNode::Evaluate(UE::Dataflow::FContext& Context, cons
 {
 	if (Out->IsA(&OutVolume))
 	{
-		if (IsConnected(&Volume))
-		{
-			const FDataflowOutput* VolumeOutput = FindOutput(&OutVolume);
+		const FDataflowOutput* VolumeOutput = FindOutput(&OutVolume);
 			
-			if (VolumeOutput)
+		if (VolumeOutput)
+		{
+			if (IsConnected(&Volume))
 			{
 				FDataflowBoolVolume* InMaskVolumePtr = nullptr;
 				if (IsConnected(&MaskVolume))
@@ -1850,63 +1335,97 @@ void FVolumeAnalysisDataflowNode::Evaluate(UE::Dataflow::FContext& Context, cons
 				}
 
 				const FDataflowVolume& InVolume = GetValue(Context, &Volume);
-				if (const FDataflowFloatVolume* InFloatVolume = InVolume.Cast<FDataflowFloatVolume>())
-				{
-					// Scalar -> Vector
-					if (Operator == EDataflowVolumeAnalysisOperator::Gradient)
-					{
-						FDataflowFloatVectorVolume GradientVolume = InFloatVolume->ComputeGradient(OutputName, CustomName, InMaskVolumePtr);
-						VolumeOutput->SetValue<FDataflowFloatVectorVolume>(MoveTemp(GradientVolume), Context);
-					}
-					// Scalar -> Scalar
-					else if (Operator == EDataflowVolumeAnalysisOperator::Curvature)
-					{
-						FDataflowFloatVolume CurvatureVolume = InFloatVolume->ComputeCurvature(OutputName, CustomName, InMaskVolumePtr);
-						VolumeOutput->SetValue<FDataflowFloatVolume>(MoveTemp(CurvatureVolume), Context);
-					}
-					// Scalar -> Scalar
-					else if (Operator == EDataflowVolumeAnalysisOperator::Laplacian)
-					{
-						FDataflowFloatVolume LaplacianVolume = InFloatVolume->ComputeLaplacian(OutputName, CustomName, InMaskVolumePtr);
-						VolumeOutput->SetValue<FDataflowFloatVolume>(MoveTemp(LaplacianVolume), Context);
-					}
-					// Scalar -> Vector
-					else if (Operator == EDataflowVolumeAnalysisOperator::ClosestPoint)
-					{
-						FDataflowFloatVectorVolume ClosestPointVolume = InFloatVolume->ComputeClosestPoint(OutputName, CustomName, InMaskVolumePtr);
-						VolumeOutput->SetValue<FDataflowFloatVectorVolume>(MoveTemp(ClosestPointVolume), Context);
-					}
-				}
-				else if (const FDataflowFloatVectorVolume* InFloatVectorVolume = InVolume.Cast<FDataflowFloatVectorVolume>())
+				if (InVolume.GetNumActiveVoxels(0.f, false))
 				{
-					// Vector -> Scalar
-					if (Operator == EDataflowVolumeAnalysisOperator::Divergence)
+					if (const FDataflowFloatVolume* InFloatVolume = InVolume.Cast<FDataflowFloatVolume>())
 					{
-						FDataflowFloatVolume DivergenceVolume = InFloatVectorVolume->ComputeDivergence(OutputName, CustomName);
-//						FDataflowFloatVolume DivergenceVolume = InFloatVectorVolume->ComputeAnalysis<FDataflowFloatVolume, "Divergence">(OutputName, CustomName, InMaskVolumePtr);
-						VolumeOutput->SetValue<FDataflowFloatVolume>(MoveTemp(DivergenceVolume), Context);
-					}
-					// Vector -> Vector
-					else if (Operator == EDataflowVolumeAnalysisOperator::Curl)
-					{
-						FDataflowFloatVectorVolume CurlVolume = InFloatVectorVolume->ComputeCurl(OutputName, CustomName);
-//						FDataflowFloatVectorVolume CurlVolume = InFloatVectorVolume->ComputeAnalysis<FDataflowFloatVectorVolume, "Curl">(OutputName, CustomName, InMaskVolumePtr);
-						VolumeOutput->SetValue<FDataflowFloatVectorVolume>(MoveTemp(CurlVolume), Context);
-					}
-					// Vector -> Scalar
-					else if (Operator == EDataflowVolumeAnalysisOperator::Magnitude)
-					{
-						FDataflowFloatVolume LengthVolume = InFloatVectorVolume->ComputeMagnitude(OutputName, CustomName);
-						VolumeOutput->SetValue<FDataflowFloatVolume>(MoveTemp(LengthVolume), Context);
+						// Scalar -> Vector
+						if (Operator == EDataflowVolumeAnalysisOperator::Gradient)
+						{
+							FDataflowFloatVectorVolume GradientVolume = InFloatVolume->ComputeGradient(OutputName, CustomName, InMaskVolumePtr);
+							VolumeOutput->SetValue<FDataflowFloatVectorVolume>(MoveTemp(GradientVolume), Context);
+							
+							return;
+						}
+						// Scalar -> Scalar
+						else if (Operator == EDataflowVolumeAnalysisOperator::Curvature)
+						{
+							FDataflowFloatVolume CurvatureVolume = InFloatVolume->ComputeCurvature(OutputName, CustomName, InMaskVolumePtr);
+							VolumeOutput->SetValue<FDataflowFloatVolume>(MoveTemp(CurvatureVolume), Context);
+
+							return;
+						}
+						// Scalar -> Scalar
+						else if (Operator == EDataflowVolumeAnalysisOperator::Laplacian)
+						{
+							FDataflowFloatVolume LaplacianVolume = InFloatVolume->ComputeLaplacian(OutputName, CustomName, InMaskVolumePtr);
+							VolumeOutput->SetValue<FDataflowFloatVolume>(MoveTemp(LaplacianVolume), Context);
+
+							return;
+						}
+						// Scalar -> Vector
+						else if (Operator == EDataflowVolumeAnalysisOperator::ClosestPoint)
+						{
+							FDataflowFloatVectorVolume ClosestPointVolume = InFloatVolume->ComputeClosestPoint(OutputName, CustomName, InMaskVolumePtr);
+							VolumeOutput->SetValue<FDataflowFloatVectorVolume>(MoveTemp(ClosestPointVolume), Context);
+
+							return;
+						}
 					}
-					// Vector -> Vector
-					else if (Operator == EDataflowVolumeAnalysisOperator::Normalize)
+					else if (const FDataflowFloatVectorVolume* InFloatVectorVolume = InVolume.Cast<FDataflowFloatVectorVolume>())
 					{
-						FDataflowFloatVectorVolume NormalVolume = InFloatVectorVolume->ComputeNormalize(OutputName, CustomName);
-						VolumeOutput->SetValue<FDataflowFloatVectorVolume>(MoveTemp(NormalVolume), Context);
+						// Vector -> Scalar
+						if (Operator == EDataflowVolumeAnalysisOperator::Divergence)
+						{
+							FDataflowFloatVolume DivergenceVolume = InFloatVectorVolume->ComputeDivergence(OutputName, CustomName);
+							//						FDataflowFloatVolume DivergenceVolume = InFloatVectorVolume->ComputeAnalysis<FDataflowFloatVolume, "Divergence">(OutputName, CustomName, InMaskVolumePtr);
+							VolumeOutput->SetValue<FDataflowFloatVolume>(MoveTemp(DivergenceVolume), Context);
+
+							return;
+						}
+						// Vector -> Vector
+						else if (Operator == EDataflowVolumeAnalysisOperator::Curl)
+						{
+							FDataflowFloatVectorVolume CurlVolume = InFloatVectorVolume->ComputeCurl(OutputName, CustomName);
+							//						FDataflowFloatVectorVolume CurlVolume = InFloatVectorVolume->ComputeAnalysis<FDataflowFloatVectorVolume, "Curl">(OutputName, CustomName, InMaskVolumePtr);
+							VolumeOutput->SetValue<FDataflowFloatVectorVolume>(MoveTemp(CurlVolume), Context);
+
+							return;
+						}
+						// Vector -> Scalar
+						else if (Operator == EDataflowVolumeAnalysisOperator::Magnitude)
+						{
+							FDataflowFloatVolume LengthVolume = InFloatVectorVolume->ComputeMagnitude(OutputName, CustomName);
+							VolumeOutput->SetValue<FDataflowFloatVolume>(MoveTemp(LengthVolume), Context);
+
+							return;
+						}
+						// Vector -> Vector
+						else if (Operator == EDataflowVolumeAnalysisOperator::Normalize)
+						{
+							FDataflowFloatVectorVolume NormalVolume = InFloatVectorVolume->ComputeNormalize(OutputName, CustomName);
+							VolumeOutput->SetValue<FDataflowFloatVectorVolume>(MoveTemp(NormalVolume), Context);
+
+							return;
+						}
 					}
 				}
 			}
+
+			if (Operator == EDataflowVolumeAnalysisOperator::Gradient ||
+				Operator == EDataflowVolumeAnalysisOperator::ClosestPoint ||
+				Operator == EDataflowVolumeAnalysisOperator::Curl ||
+				Operator == EDataflowVolumeAnalysisOperator::Normalize)
+			{
+				VolumeOutput->SetValue<FDataflowFloatVectorVolume>(FDataflowFloatVectorVolume(), Context);
+			}
+			else if (Operator == EDataflowVolumeAnalysisOperator::Curvature ||
+				Operator == EDataflowVolumeAnalysisOperator::Laplacian ||
+				Operator == EDataflowVolumeAnalysisOperator::Divergence ||
+				Operator == EDataflowVolumeAnalysisOperator::Magnitude)
+			{
+				VolumeOutput->SetValue<FDataflowFloatVolume>(FDataflowFloatVolume(), Context);
+			}
 		}
 	}
 	else if (Out->IsA(&MaskVolume))
@@ -1971,58 +1490,6 @@ void FVolumeCombineDataflowNode::Evaluate(UE::Dataflow::FContext& Context, const
 	}
 }
 
-#if WITH_EDITOR
-bool FVolumeCombineDataflowNode::CanDebugDrawViewMode(const FName& ViewModeName) const
-{
-	return ViewModeName == UE::Dataflow::FDataflowConstruction3DViewMode::Name;
-}
-
-void FVolumeCombineDataflowNode::DebugDraw(UE::Dataflow::FContext& Context, IDataflowDebugDrawInterface& DataflowRenderingInterface, const FDebugDrawParameters& DebugDrawParameters) const
-{
-	using namespace UE::Dataflow::Overlay;
-
-	if (DebugDrawParameters.bNodeIsSelected || DebugDrawParameters.bNodeIsPinned)
-	{
-		FString OverlayEvaluateResultString;
-
-		const FDataflowOutput* FloatVolumeOutput = FindOutput(&FloatVolume);
-		if (FloatVolumeOutput)
-		{
-			const FDataflowFloatVolume& OutFloatVolume = FloatVolumeOutput->GetValue(Context, FloatVolume);
-
-			DebugDrawRenderSettings.SetDebugDrawSettings(DataflowRenderingInterface);
-
-			TArray<FBox> ActiveVoxels;
-			OutFloatVolume.GetActiveVoxels(ActiveVoxels);
-
-			for (const FBox& Voxel : ActiveVoxels)
-			{
-				if (bDrawVoxelCenters)
-				{
-					DataflowRenderingInterface.DrawPoint(0.5 * FVector(Voxel.Max + Voxel.Min));
-				}
-
-				if (bDrawVoxels)
-				{
-					DataflowRenderingInterface.DrawBox(0.5 * FVector(Voxel.Max - Voxel.Min), FQuat::Identity, 0.5 * FVector(Voxel.Min + Voxel.Max), 1.0);
-				}
-			}
-
-			OverlayEvaluateResultString = OutFloatVolume.VolumeInfo();
-		}
-
-		// ------------------------------------------------------------------------------------------------------
-		// Overlay
-		TArray<FString> OverlayStrArr;
-		OverlayStrArr.Add(BuildOverlayNodeInfoString(this));
-		OverlayStrArr.Add(OverlayEvaluateResultString);
-		FString OverlayStr = BuildOverlayFinalString(OverlayStrArr);
-		DataflowRenderingInterface.DrawOverlayText(OverlayStr);
-		// ------------------------------------------------------------------------------------------------------
-	}
-}
-#endif
-
 /* --------------------------------------------------------------------------------------------- */
 
 namespace DataflowVolumeNodes::Private
@@ -2080,7 +1547,7 @@ void FVolumeTextureTerminalNode::SetAssetValue(TObjectPtr<UObject> Asset, UE::Da
 	{
 		const FDataflowVolume& InVolume = GetValue(Context, &Volume);
 
-		if (InVolume.GetNumActiveVoxels() > 0)
+		if (InVolume.GetNumActiveVoxels(0.f, false) > 0)
 		{
 			DataflowVolumeNodes::Private::UpdateVolumeTextureFromVolume(AssetToSet, InVolume);
 		}
diff --git a/Engine/Plugins/Experimental/DataflowVolumetricPlugin/Source/DataflowVolumeNodes/Private/Dataflow/DataflowVolumeNodes.h b/Engine/Plugins/Experimental/DataflowVolumetricPlugin/Source/DataflowVolumeNodes/Private/Dataflow/DataflowVolumeNodes.h
--- a/Engine/Plugins/Experimental/DataflowVolumetricPlugin/Source/DataflowVolumeNodes/Private/Dataflow/DataflowVolumeNodes.h
+++ b/Engine/Plugins/Experimental/DataflowVolumetricPlugin/Source/DataflowVolumeNodes/Private/Dataflow/DataflowVolumeNodes.h
@@ -38,13 +38,15 @@ private:
 	UPROPERTY(EditAnywhere, Category = "Sphere", meta = (DataflowInput));
 	FVector Center = FVector(0);
 
-	UPROPERTY(EditAnywhere, Category = "Debug Draw", meta = (ShowOnlyInnerProperties))
+	// These properties are not needed anymore as we use the new rendering system. 
+	// Since we don't wan't to version up the node we will just hide them.
+	UPROPERTY()
 	FDataflowNodeDebugDrawSettings DebugDrawRenderSettings;
 
-	UPROPERTY(EditAnywhere, Category = "Debug Draw")
+	UPROPERTY()
 	bool bDrawVoxels = true;
 
-	UPROPERTY(EditAnywhere, Category = "Debug Draw")
+	UPROPERTY()
 	bool bDrawVoxelCenters = false;
 
 	/** SDF volume */
@@ -53,15 +55,6 @@ private:
 
 	virtual void Evaluate(UE::Dataflow::FContext& Context, const FDataflowOutput* Out) const override;
 
-#if WITH_EDITOR
-	virtual bool CanDebugDraw() const override
-	{
-		return true;
-	}
-	virtual bool CanDebugDrawViewMode(const FName& ViewModeName) const override;
-	virtual void DebugDraw(UE::Dataflow::FContext& Context, IDataflowDebugDrawInterface& DataflowRenderingInterface, const FDebugDrawParameters& DebugDrawParameters) const override;
-#endif
-
 public:
 	FMakeSDFSphereDataflowNode(const UE::Dataflow::FNodeParameters& InParam, FGuid InGuid = FGuid::NewGuid());
 };
@@ -87,14 +80,16 @@ private:
 	/** The size of voxels in the output VDB volume */
 	UPROPERTY(EditAnywhere, Category = "SDF", meta = (DataflowInput, ClampMin = 0.1f));
 	float VoxelSize = 1.f;
-
-	UPROPERTY(EditAnywhere, Category = "Debug Draw", meta = (ShowOnlyInnerProperties))
+	
+	// These properties are not needed anymore as we use the new rendering system. 
+	// Since we don't wan't to version up the node we will just hide them.
+	UPROPERTY()
 	FDataflowNodeDebugDrawSettings DebugDrawRenderSettings;
 
-	UPROPERTY(EditAnywhere, Category = "Debug Draw")
+	UPROPERTY()
 	bool bDrawVoxels = true;
 
-	UPROPERTY(EditAnywhere, Category = "Debug Draw")
+	UPROPERTY()
 	bool bDrawVoxelCenters = false;
 
 	/** SDF volume */
@@ -103,15 +98,6 @@ private:
 
 	virtual void Evaluate(UE::Dataflow::FContext& Context, const FDataflowOutput* Out) const override;
 
-#if WITH_EDITOR
-	virtual bool CanDebugDraw() const override
-	{
-		return true;
-	}
-	virtual bool CanDebugDrawViewMode(const FName& ViewModeName) const override;
-	virtual void DebugDraw(UE::Dataflow::FContext& Context, IDataflowDebugDrawInterface& DataflowRenderingInterface, const FDebugDrawParameters& DebugDrawParameters) const override;
-#endif
-
 public:
 	FMakeSDFCubeDataflowNode(const UE::Dataflow::FNodeParameters& InParam, FGuid InGuid = FGuid::NewGuid());
 };
@@ -142,13 +128,15 @@ private:
 	UPROPERTY(EditAnywhere, Category = "SDF", meta = (DataflowInput, ClampMin = 0.1f));
 	float VoxelSize = 1.f;
 
-	UPROPERTY(EditAnywhere, Category = "Debug Draw", meta = (ShowOnlyInnerProperties))
+	// These properties are not needed anymore as we use the new rendering system. 
+	// Since we don't wan't to version up the node we will just hide them.
+	UPROPERTY()
 	FDataflowNodeDebugDrawSettings DebugDrawRenderSettings;
 
-	UPROPERTY(EditAnywhere, Category = "Debug Draw")
+	UPROPERTY()
 	bool bDrawVoxels = true;
 
-	UPROPERTY(EditAnywhere, Category = "Debug Draw")
+	UPROPERTY()
 	bool bDrawVoxelCenters = false;
 
 	/** SDF volume */
@@ -157,15 +145,6 @@ private:
 
 	virtual void Evaluate(UE::Dataflow::FContext& Context, const FDataflowOutput* Out) const override;
 
-#if WITH_EDITOR
-	virtual bool CanDebugDraw() const override
-	{
-		return true;
-	}
-	virtual bool CanDebugDrawViewMode(const FName& ViewModeName) const override;
-	virtual void DebugDraw(UE::Dataflow::FContext& Context, IDataflowDebugDrawInterface& DataflowRenderingInterface, const FDebugDrawParameters& DebugDrawParameters) const override;
-#endif
-
 public:
 	FMakeSDFPlatonicSolidDataflowNode(const UE::Dataflow::FNodeParameters& InParam, FGuid InGuid = FGuid::NewGuid());
 };
@@ -262,13 +241,15 @@ private:
 	UPROPERTY(EditAnywhere, Category = "Closest Surface Point", meta = (DataflowInput));
 	float IsoValue = 0.f;
 
-	UPROPERTY(EditAnywhere, Category = "Debug Draw", meta = (ShowOnlyInnerProperties))
+	// These properties are not needed anymore as we use the new rendering system. 
+	// Since we don't wan't to version up the node we will just hide them.
+	UPROPERTY()
 	FDataflowNodeDebugDrawSettings DebugDrawRenderSettings;
 
-	UPROPERTY(EditAnywhere, Category = "Debug Draw")
+	UPROPERTY()
 	bool bDrawVoxels = false;
 
-	UPROPERTY(EditAnywhere, Category = "Debug Draw")
+	UPROPERTY()
 	bool bDrawVoxelCenters = true;
 
 	/** Float volume */
@@ -281,15 +262,6 @@ private:
 
 	virtual void Evaluate(UE::Dataflow::FContext& Context, const FDataflowOutput* Out) const override;
 
-#if WITH_EDITOR
-	virtual bool CanDebugDraw() const override
-	{
-		return true;
-	}
-	virtual bool CanDebugDrawViewMode(const FName& ViewModeName) const override;
-	virtual void DebugDraw(UE::Dataflow::FContext& Context, IDataflowDebugDrawInterface& DataflowRenderingInterface, const FDebugDrawParameters& DebugDrawParameters) const override;
-#endif
-
 public:
 	FCollectionToVolumeDataflowNode(const UE::Dataflow::FNodeParameters& InParam, FGuid InGuid = FGuid::NewGuid());
 };
@@ -331,13 +303,15 @@ private:
 	UPROPERTY(EditAnywhere, Category = "Spheres", meta = (DataflowInput, ClampMin = 1));
 	int32 InstanceCount = 10000;
 
-	UPROPERTY(EditAnywhere, Category = "Debug Draw", meta = (ShowOnlyInnerProperties))
+	// These properties are not needed anymore as we use the new rendering system. 
+	// Since we don't wan't to version up the node we will just hide them.
+	UPROPERTY()
 	FDataflowNodeDebugDrawSettings DebugDrawRenderSettings;
 
-	UPROPERTY(EditAnywhere, Category = "Debug Draw")
+	UPROPERTY()
 	bool bDrawSpheres = true;
 
-	UPROPERTY(EditAnywhere, Category = "Debug Draw")
+	UPROPERTY()
 	bool bDrawSphereCenters = false;
 
 	/**  */
@@ -350,15 +324,6 @@ private:
 
 	virtual void Evaluate(UE::Dataflow::FContext& Context, const FDataflowOutput* Out) const override;
 
-#if WITH_EDITOR
-	virtual bool CanDebugDraw() const override
-	{
-		return true;
-	}
-	virtual bool CanDebugDrawViewMode(const FName& ViewModeName) const override;
-	virtual void DebugDraw(UE::Dataflow::FContext& Context, IDataflowDebugDrawInterface& DataflowRenderingInterface, const FDebugDrawParameters& DebugDrawParameters) const override;
-#endif
-
 public:
 	FVolumeToSpheresDataflowNode(const UE::Dataflow::FNodeParameters& InParam, FGuid InGuid = FGuid::NewGuid());
 };
@@ -418,13 +383,15 @@ private:
 	UPROPERTY(EditAnywhere, Category = "Scatter", meta = (DataflowInput, ClampMin = 0, EditCondition = "ScatterType != EDataflowVolumeScatterType::Uniform", EditConditionHides));
 	float MaxNumberOfPointsPerVoxel = 1.f;
 
-	UPROPERTY(EditAnywhere, Category = "Debug Draw", meta = (ShowOnlyInnerProperties))
+	// These properties are not needed anymore as we use the new rendering system. 
+	// Since we don't wan't to version up the node we will just hide them.
+	UPROPERTY()
 	FDataflowNodeDebugDrawSettings DebugDrawRenderSettings;
 
-	UPROPERTY(EditAnywhere, Category = "Debug Draw")
+	UPROPERTY()
 	bool bDrawSpheres = true;
 
-	UPROPERTY(EditAnywhere, Category = "Debug Draw")
+	UPROPERTY()
 	bool bDrawSphereCenters = false;
 
 	/**  */
@@ -485,13 +452,15 @@ private:
 	UPROPERTY(EditAnywhere, Category = "Convert", meta = (DisplayName = "Activate Interior Voxels", EditCondition = "ConvertTo == EDataflowVolumeConvertSDFTo::SDF", EditConditionHides));
 	bool bActivateInteriror = true;
 
-	UPROPERTY(EditAnywhere, Category = "Debug Draw", meta = (ShowOnlyInnerProperties))
+	// These properties are not needed anymore as we use the new rendering system. 
+	// Since we don't wan't to version up the node we will just hide them.
+	UPROPERTY()
 	FDataflowNodeDebugDrawSettings DebugDrawRenderSettings;
 
-	UPROPERTY(EditAnywhere, Category = "Debug Draw")
+	UPROPERTY()
 	bool bDrawVoxels = true;
 
-	UPROPERTY(EditAnywhere, Category = "Debug Draw")
+	UPROPERTY()
 	bool bDrawVoxelCenters = false;
 
 	/** Collection output */
@@ -500,14 +469,6 @@ private:
 
 	virtual void Evaluate(UE::Dataflow::FContext& Context, const FDataflowOutput* Out) const override;
 
-#if WITH_EDITOR
-	virtual bool CanDebugDraw() const override
-	{
-		return true;
-	}
-	virtual bool CanDebugDrawViewMode(const FName& ViewModeName) const override;
-	virtual void DebugDraw(UE::Dataflow::FContext& Context, IDataflowDebugDrawInterface& DataflowRenderingInterface, const FDebugDrawParameters& DebugDrawParameters) const override;
-#endif
 public:
 	FConvertVolumeDataflowNode(const UE::Dataflow::FNodeParameters& InParam, FGuid InGuid = FGuid::NewGuid());
 };
@@ -546,7 +507,9 @@ private:
 	UPROPERTY(EditAnywhere, Category = "Volume", meta = (DataflowInput, ClampMin = 0, ClampMax = 1));
 	float Spread = 1.f;
 
-	UPROPERTY(EditAnywhere, Category = "Debug Draw")
+	// These properties are not needed anymore as we use the new rendering system. 
+	// Since we don't wan't to version up the node we will just hide them.
+	UPROPERTY()
 	FLinearColor Color = FLinearColor(0.7, 0.0, 1.0, 1.0);
 
 	/** Generated points */
@@ -555,15 +518,6 @@ private:
 
 	virtual void Evaluate(UE::Dataflow::FContext& Context, const FDataflowOutput* Out) const override;
 
-#if WITH_EDITOR
-	virtual bool CanDebugDraw() const override
-	{
-		return true;
-	}
-	virtual bool CanDebugDrawViewMode(const FName& ViewModeName) const override;
-	virtual void DebugDraw(UE::Dataflow::FContext& Context, IDataflowDebugDrawInterface& DataflowRenderingInterface, const FDebugDrawParameters& DebugDrawParameters) const override;
-#endif
-
 public:
 	FUniformVolumeScatterDataflowNode(const UE::Dataflow::FNodeParameters& InParam, FGuid InGuid = FGuid::NewGuid());
 };
@@ -606,7 +560,9 @@ private:
 	UPROPERTY(EditAnywhere, Category = "Volume", meta = (DataflowInput, ClampMin = 0, ClampMax = 1));
 	float Spread = 1.f;
 
-	UPROPERTY(EditAnywhere, Category = "Debug Draw")
+	// These properties are not needed anymore as we use the new rendering system. 
+	// Since we don't wan't to version up the node we will just hide them.
+	UPROPERTY()
 	FLinearColor Color = FLinearColor(0.7, 0.0, 1.0, 1.0);
 
 	/** Generated points */
@@ -615,15 +571,6 @@ private:
 
 	virtual void Evaluate(UE::Dataflow::FContext& Context, const FDataflowOutput* Out) const override;
 
-#if WITH_EDITOR
-	virtual bool CanDebugDraw() const override
-	{
-		return true;
-	}
-	virtual bool CanDebugDrawViewMode(const FName& ViewModeName) const override;
-	virtual void DebugDraw(UE::Dataflow::FContext& Context, IDataflowDebugDrawInterface& DataflowRenderingInterface, const FDebugDrawParameters& DebugDrawParameters) const override;
-#endif
-
 public:
 	FDenseUniformVolumeScatterDataflowNode(const UE::Dataflow::FNodeParameters& InParam, FGuid InGuid = FGuid::NewGuid());
 };
@@ -667,7 +614,9 @@ private:
 	UPROPERTY(EditAnywhere, Category = "Volume", meta = (DataflowInput, ClampMin = 0, ClampMax = 1));
 	float Spread = 1.f;
 
-	UPROPERTY(EditAnywhere, Category = "Debug Draw")
+	// These properties are not needed anymore as we use the new rendering system. 
+	// Since we don't wan't to version up the node we will just hide them.
+	UPROPERTY()
 	FLinearColor Color = FLinearColor(0.7, 0.0, 1.0, 1.0);
 
 	/** Generated points */
@@ -676,15 +625,6 @@ private:
 
 	virtual void Evaluate(UE::Dataflow::FContext& Context, const FDataflowOutput* Out) const override;
 
-#if WITH_EDITOR
-	virtual bool CanDebugDraw() const override
-	{
-		return true;
-	}
-	virtual bool CanDebugDrawViewMode(const FName& ViewModeName) const override;
-	virtual void DebugDraw(UE::Dataflow::FContext& Context, IDataflowDebugDrawInterface& DataflowRenderingInterface, const FDebugDrawParameters& DebugDrawParameters) const override;
-#endif
-
 public:
 	FNonUniformVolumeScatterDataflowNode(const UE::Dataflow::FNodeParameters& InParam, FGuid InGuid = FGuid::NewGuid());
 };
@@ -710,7 +650,9 @@ private:
 	UPROPERTY(EditAnywhere, Category = "Sampling");
 	EDataflowVolumeSampleType SampleType = EDataflowVolumeSampleType::Float;
 
-	UPROPERTY(EditAnywhere, Category = "Debug Draw")
+	// These properties are not needed anymore as we use the new rendering system. 
+	// Since we don't wan't to version up the node we will just hide them.
+	UPROPERTY()
 	FLinearColor Color = FLinearColor(0.2, 1.0, 1.0, 1.0);
 
 	UPROPERTY(EditAnywhere, Category = "Sampling")
@@ -724,14 +666,7 @@ private:
 	FManagedArrayCollection PointsCollection;
 
 	virtual void Evaluate(UE::Dataflow::FContext& Context, const FDataflowOutput* Out) const override;
-#if WITH_EDITOR
-	virtual bool CanDebugDraw() const override
-	{
-		return true;
-	}
-	virtual bool CanDebugDrawViewMode(const FName& ViewModeName) const override;
-	virtual void DebugDraw(UE::Dataflow::FContext& Context, IDataflowDebugDrawInterface& DataflowRenderingInterface, const FDebugDrawParameters& DebugDrawParameters) const override;
-#endif
+
 public:
 	FVolumeSampleToPointsDataflowNode(const UE::Dataflow::FNodeParameters& InParam, FGuid InGuid = FGuid::NewGuid());
 
@@ -909,13 +844,15 @@ private:
 	UPROPERTY(EditAnywhere, Category = "Combine", meta = (DisplayName = "Prune Tolerance", UIMin = 0.f, UIMax = 1.f));
 	float PruneTol = 0.f;
 
-	UPROPERTY(EditAnywhere, Category = "Debug Draw", meta = (ShowOnlyInnerProperties))
+	// These properties are not needed anymore as we use the new rendering system. 
+	// Since we don't wan't to version up the node we will just hide them.
+	UPROPERTY()
 	FDataflowNodeDebugDrawSettings DebugDrawRenderSettings;
 
-	UPROPERTY(EditAnywhere, Category = "Debug Draw")
+	UPROPERTY()
 	bool bDrawVoxels = true;
 
-	UPROPERTY(EditAnywhere, Category = "Debug Draw")
+	UPROPERTY()
 	bool bDrawVoxelCenters = false;
 
 	/** Output Float volume */
@@ -923,14 +860,6 @@ private:
 	FDataflowFloatVolume FloatVolume;
 
 	virtual void Evaluate(UE::Dataflow::FContext& Context, const FDataflowOutput* Out) const override;
-#if WITH_EDITOR
-	virtual bool CanDebugDraw() const override
-	{
-		return true;
-	}
-	virtual bool CanDebugDrawViewMode(const FName& ViewModeName) const override;
-	virtual void DebugDraw(UE::Dataflow::FContext& Context, IDataflowDebugDrawInterface& DataflowRenderingInterface, const FDebugDrawParameters& DebugDrawParameters) const override;
-#endif
 
 public:
 	FVolumeCombineDataflowNode(const UE::Dataflow::FNodeParameters& InParam, FGuid InGuid = FGuid::NewGuid());
@@ -1011,8 +940,8 @@ private:
 	UPROPERTY(EditAnywhere, Category = "Relaxation");
 	bool bApplyRelaxation = true;
 
-	/** Number of steps in the relaxation */
-	UPROPERTY(EditAnywhere, Category = "Relaxation", meta = (DataflowInput, UIMin = 1, ClampMin = 1, EditCondition = "bApplyRelaxation"));
+	/** Number of iterations for the relaxation */
+	UPROPERTY(EditAnywhere, Category = "Relaxation", meta = (DataflowInput, DisplayName = "Iteration", UIMin = 1, ClampMin = 1, UIMax = 20, EditCondition = "bApplyRelaxation"));
 	int32 Steps = 10;
 
 	/** Scale the displacement applied to separate spheres */
@@ -1037,27 +966,6 @@ public:
 	FVolumeToSpheresWithRelaxationDataflowNode(const UE::Dataflow::FNodeParameters& InParam, FGuid InGuid = FGuid::NewGuid());
 };
 
-
-
-
-
-
-
-
-
-
-
-
-
-
-
-
-
-
-
-
-
-
 namespace UE::Dataflow
 {
 	void RegisterVolumeNodes();
