Dataflow:
- Added rendering support for FDataflowVolume (it renders the voxel centers as points and the outer layer of voxels as boxes) to the new rendering system

#Synced-CL 50147825

diff --git a/Engine/Plugins/Experimental/Dataflow/Source/DataflowEditor/Public/DataflowRendering/DataflowRenderableTypeUtils.h b/Engine/Plugins/Experimental/Dataflow/Source/DataflowEditor/Public/DataflowRendering/DataflowRenderableTypeUtils.h
--- a/Engine/Plugins/Experimental/Dataflow/Source/DataflowEditor/Public/DataflowRendering/DataflowRenderableTypeUtils.h
+++ b/Engine/Plugins/Experimental/Dataflow/Source/DataflowEditor/Public/DataflowRendering/DataflowRenderableTypeUtils.h
@@ -52,7 +52,7 @@ struct FDataflowColorCommonRenderSettings
 	GENERATED_BODY()
 
 public:
-	FDataflowColorCommonRenderSettings();
+	UE_API FDataflowColorCommonRenderSettings();
 
 	/** Display surface as wireframe */
 	UPROPERTY(EditAnywhere, Category = "Color")
@@ -144,7 +144,7 @@ namespace UE::Dataflow
 
 		int32 GetUVChannelFromInstance(const FRenderableTypeInstance& Instance, int32 DefaultValue = 0);
 
-		FLinearColor GetColor(const int32 InIdx, const FBox& InComponentBoundingBox, const FVector& InCenter, float InSize,
+		UE_API FLinearColor GetColor(const int32 InIdx, const FBox& InComponentBoundingBox, const FVector& InCenter, float InSize,
 			const FDataflowColorCommonRenderSettings& InSettings);
 	}
 }
diff --git a/Engine/Plugins/Experimental/DataflowVolumetricPlugin/Source/DataflowVolumeCore/Internal/Dataflow/DataflowVolumeRenderableType.cpp b/Engine/Plugins/Experimental/DataflowVolumetricPlugin/Source/DataflowVolumeCore/Internal/Dataflow/DataflowVolumeRenderableType.cpp
new file mode 100644
--- /dev/null
+++ b/Engine/Plugins/Experimental/DataflowVolumetricPlugin/Source/DataflowVolumeCore/Internal/Dataflow/DataflowVolumeRenderableType.cpp
@@ -0,0 +1,217 @@
+// Copyright Epic Games, Inc. All Rights Reserved.
+
+#include "Dataflow/DataflowVolumeRenderableType.h"
+#include "DataflowRendering/DataflowRenderableTypeUtils.h"
+#include "Dataflow/DataflowEngineUtil.h"
+
+#include "Components/BoxComponent.h"
+#include "Components/InstancedStaticMeshComponent.h"
+#include "Drawing/PointSetComponent.h"
+
+#include "Dataflow/DataflowRenderingViewMode.h"
+#include "Dataflow/DataflowVolume.h"
+
+#include "DataflowRendering/DataflowRenderableComponents.h"
+#include "DataflowRendering/DataflowRenderableType.h"
+#include "DataflowRendering/DataflowRenderableTypeInstance.h"
+#include "DataflowRendering/DataflowRenderableTypeRegistry.h"
+
+#include "Materials/MaterialInstanceDynamic.h"
+
+#include "UObject/ObjectPtr.h"
+
+#include "Dataflow/DataflowFloatVolume.h"
+#include "Dataflow/DataflowFloatVectorVolume.h"
+#include "Dataflow/DataflowBoolVolume.h"
+#include "Dataflow/DataflowIntVolume.h"
+
+UDataflowVolumeRenderSettings::UDataflowVolumeRenderSettings()
+{
+	ColorSettings.bWireframe = false;
+	ColorSettings.Transparency = 0.6f;
+	ColorSettings.Color = FLinearColor(FVector(0.624f, 0.465f, 0.776f));
+}
+
+namespace UE::Dataflow::Private
+{
+	///////////////////////////////////////////////////////////////////////////////////////////////////////////////
+
+	class FVolumeSurfaceRenderableType : public IRenderableType
+	{
+		virtual const FDataflowVolume& GetVolume(const FRenderableTypeInstance& Instance) const = 0;
+		virtual const FName GetComponentName() const = 0;
+
+		UE_DATAFLOW_IRENDERABLE_RENDER_GROUP(Surface);
+		UE_DATAFLOW_IRENDERABLE_VIEW_MODE(FDataflowConstruction3DViewMode);
+		UE_DATAFLOW_IRENDERABLE_SETTINGS(UDataflowVolumeRenderSettings);
+
+		virtual void GetPrimitiveComponents(const FRenderableTypeInstance& Instance, FRenderableComponents& OutComponents) const override
+		{
+			UStaticMesh* BoxMesh = UE::Dataflow::RenderGeometry::GetDataflowBox();
+			if (BoxMesh)
+			{
+				const FDataflowVolume& Volume = GetVolume(Instance);
+
+				const UDataflowVolumeRenderSettings* Settings = Instance.GetTypedRenderSettings<UDataflowVolumeRenderSettings>();
+				if (Settings)
+				{
+					TArray<FBox> BoxArray;
+					Volume.GetActiveVoxels(Settings->IsoValue, BoxArray, true);
+
+					if (BoxArray.Num())
+					{
+						const FName BoxArrayComponentName = GetComponentName();
+						const int32 NumBoxes = BoxArray.Num();
+
+						UInstancedStaticMeshComponent* ISMComponent = OutComponents.AddNewComponent<UInstancedStaticMeshComponent>(BoxArrayComponentName);
+						if (ISMComponent)
+						{
+							ISMComponent->PreAllocateInstancesMemory(NumBoxes);
+							ISMComponent->SetStaticMesh(BoxMesh);
+
+							ISMComponent->SetCollisionEnabled(ECollisionEnabled::NoCollision);
+
+							FBox ComponentBoundingBox;
+							for (const FBox& Box : BoxArray)
+							{
+								ComponentBoundingBox += Box;
+
+								FTransform BoxTransform;
+								BoxTransform.SetTranslation(Box.GetCenter());
+								BoxTransform.SetScale3D(Settings->VoxelScale * Box.GetExtent() / 50.f);
+
+								ISMComponent->AddInstance(BoxTransform);
+							}
+
+							ISMComponent->SetNumCustomDataFloats(3); // RGB
+
+							UMaterialInterface* Material = Settings->ColorSettings.bWireframe
+								? UE::Dataflow::RenderMaterial::GetDataflowColorWireframeMaterial()
+								: UE::Dataflow::RenderMaterial::GetDataflowColorMaterial();
+							if (Material)
+							{
+								UMaterialInstanceDynamic* MaterialInstance = UMaterialInstanceDynamic::Create(Material, ISMComponent);
+								if (MaterialInstance)
+								{
+									ISMComponent->SetMaterial(0, MaterialInstance);
+
+									MaterialInstance->SetScalarParameterValue(FName("Transparency"), Settings->ColorSettings.Transparency);
+
+									TArray<float> Colors; // NumInstances * 3
+									Colors.Init(1.0, BoxArray.Num() * 3);
+
+									int32 Idx = 0;
+									for (const FBox& Box : BoxArray)
+									{
+										FLinearColor Color = UE::Dataflow::Rendering::GetColor(Idx,
+											ComponentBoundingBox,
+											Box.GetCenter(),
+											2.0 * Box.GetExtent().GetMax(),
+											Settings->ColorSettings);
+
+										Colors[3 * Idx + 0] = Color.R;
+										Colors[3 * Idx + 1] = Color.G;
+										Colors[3 * Idx + 2] = Color.B;
+
+										Idx++;
+									}
+
+									ISMComponent->SetCustomData(0, BoxArray.Num() - 1, Colors);
+								}
+							}
+						}
+					}
+				}
+			}
+		}
+	};
+
+	class FFloatVolumeSurfaceRenderableType : public FVolumeSurfaceRenderableType
+	{
+		virtual const FDataflowVolume& GetVolume(const FRenderableTypeInstance& Instance) const override
+		{
+			static FDataflowFloatVolume Default;
+
+			return GetFloatVolume(Instance, Default);
+		}
+
+		virtual const FName GetComponentName() const override
+		{
+			static FName ComponentName = FName("FloatVolume");
+
+			return ComponentName;
+		}
+
+		UE_DATAFLOW_IRENDERABLE_PRIMARY_TYPE(FDataflowFloatVolume, FloatVolume);
+	};
+
+	class FIntVolumeSurfaceRenderableType : public FVolumeSurfaceRenderableType
+	{
+		virtual const FDataflowVolume& GetVolume(const FRenderableTypeInstance& Instance) const override
+		{
+			static FDataflowIntVolume Default;
+
+			return GetIntVolume(Instance, Default);
+		}
+
+		virtual const FName GetComponentName() const override
+		{
+			static FName ComponentName = FName("IntVolume");
+
+			return ComponentName;
+		}
+
+		UE_DATAFLOW_IRENDERABLE_PRIMARY_TYPE(FDataflowIntVolume, IntVolume);
+	};
+
+	class FBoolVolumeSurfaceRenderableType : public FVolumeSurfaceRenderableType
+	{
+		virtual const FDataflowVolume& GetVolume(const FRenderableTypeInstance& Instance) const override
+		{
+			static FDataflowBoolVolume Default;
+
+			return GetBoolVolume(Instance, Default);
+		}
+
+		virtual const FName GetComponentName() const override
+		{
+			static FName ComponentName = FName("BoolVolume");
+
+			return ComponentName;
+		}
+
+		UE_DATAFLOW_IRENDERABLE_PRIMARY_TYPE(FDataflowBoolVolume, BoolVolume);
+	};
+
+	class FFloatVectorVolumeSurfaceRenderableType : public FVolumeSurfaceRenderableType
+	{
+		virtual const FDataflowVolume& GetVolume(const FRenderableTypeInstance& Instance) const override
+		{
+			static FDataflowFloatVectorVolume Default;
+
+			return GetFloatVectorVolume(Instance, Default);
+		}
+
+		virtual const FName GetComponentName() const override
+		{
+			static FName ComponentName = FName("FloatVectorVolume");
+
+			return ComponentName;
+		}
+
+		UE_DATAFLOW_IRENDERABLE_PRIMARY_TYPE(FDataflowFloatVectorVolume, FloatVectorVolume);
+	};
+
+	///////////////////////////////////////////////////////////////////////////////////////////////////////////////
+
+	void RegisterVolumeRenderableTypes()
+	{
+		UE_DATAFLOW_REGISTER_RENDERABLE_TYPE(FFloatVolumeSurfaceRenderableType);
+		UE_DATAFLOW_REGISTER_RENDERABLE_TYPE(FIntVolumeSurfaceRenderableType);
+		UE_DATAFLOW_REGISTER_RENDERABLE_TYPE(FBoolVolumeSurfaceRenderableType);
+		UE_DATAFLOW_REGISTER_RENDERABLE_TYPE(FFloatVectorVolumeSurfaceRenderableType);
+	}
+
+	///////////////////////////////////////////////////////////////////////////////////////////////////////////////
+
+}
\ No newline at end of file
diff --git a/Engine/Plugins/Experimental/DataflowVolumetricPlugin/Source/DataflowVolumeCore/Internal/Dataflow/DataflowVolumeRenderableType.h b/Engine/Plugins/Experimental/DataflowVolumetricPlugin/Source/DataflowVolumeCore/Internal/Dataflow/DataflowVolumeRenderableType.h
new file mode 100644
--- /dev/null
+++ b/Engine/Plugins/Experimental/DataflowVolumetricPlugin/Source/DataflowVolumeCore/Internal/Dataflow/DataflowVolumeRenderableType.h
@@ -0,0 +1,33 @@
+// Copyright Epic Games, Inc. All Rights Reserved.
+
+#pragma once
+
+#include "DataflowRendering/DataflowRenderableTypeSettings.h"
+#include "DataflowRendering/DataflowRenderableTypeUtils.h"
+
+#include "DataflowVolumeRenderableType.generated.h"
+
+UCLASS(MinimalAPI, AutoExpandCategories = "Surface|ConstantColor, Surface|RandomColor, Surface|ColorBySize")
+class UDataflowVolumeRenderSettings : public UDataflowRenderableTypeSettings
+{
+	GENERATED_BODY()
+
+public:
+	UDataflowVolumeRenderSettings();
+
+	UPROPERTY(EditAnywhere, Category = "Volume", meta = (ShowOnlyInnerProperties))
+	FDataflowColorCommonRenderSettings ColorSettings;
+
+	/** Voxels to display the volume gets generated at the Isovalue */
+	UPROPERTY(EditAnywhere, Category = "Volume", meta = (UIMin = -3, UIMax = 3))
+	float IsoValue = 0.f;
+
+	/** Scale applied to display the voxels */
+	UPROPERTY(EditAnywhere, Category = "Volume", meta = (UIMin = 0.001f, ClampMin = 0.001f, UIMax = 1.f, ClampMax = 1.f))
+	float VoxelScale = .75f;
+};
+
+namespace UE::Dataflow::Private
+{
+	void RegisterVolumeRenderableTypes();
+}
diff --git a/Engine/Plugins/Experimental/DataflowVolumetricPlugin/Source/DataflowVolumeCore/Private/Dataflow/DataflowVolumeCoreModule.cpp b/Engine/Plugins/Experimental/DataflowVolumetricPlugin/Source/DataflowVolumeCore/Private/Dataflow/DataflowVolumeCoreModule.cpp
--- a/Engine/Plugins/Experimental/DataflowVolumetricPlugin/Source/DataflowVolumeCore/Private/Dataflow/DataflowVolumeCoreModule.cpp
+++ b/Engine/Plugins/Experimental/DataflowVolumetricPlugin/Source/DataflowVolumeCore/Private/Dataflow/DataflowVolumeCoreModule.cpp
@@ -7,6 +7,7 @@
 #include "Modules/ModuleInterface.h"
 #include "Modules/ModuleManager.h"
 #include "Dataflow/OpenVDB.h"
+#include "Dataflow/DataflowVolumeRenderableType.h"
 
 #define LOCTEXT_NAMESPACE "DataflowCore"
 
@@ -14,8 +15,12 @@ void IDataflowVolumeCoreModule::StartupModule()
 {
 	// Initialize OpenVDB
 	openvdb::initialize();
+
+	// Register Volume render type
+	UE::Dataflow::Private::RegisterVolumeRenderableTypes();
 }
 
+
 void IDataflowVolumeCoreModule::ShutdownModule()
 {}
 
